chimera - move to //UE5/Main/Engine/Plugins/Experimental/Animation

diff --git a/Engine/Plugins/Experimental/Animation/Chimera/Chimera.uplugin b/Engine/Plugins/Experimental/Animation/Chimera/Chimera.uplugin
new file mode 100644
--- /dev/null
+++ b/Engine/Plugins/Experimental/Animation/Chimera/Chimera.uplugin
@@ -0,0 +1,37 @@
+{
+	"FileVersion": 3,
+	"Version": 1,
+	"VersionName": "1.0",
+	"FriendlyName": "Chimera",
+	"Description": "Character Interaction Matching",
+	"Category": "Animation",
+	"CreatedBy": "Epic Games, Inc.",
+	"CreatedByURL": "https://epicgames.com",
+	"DocsURL": "",
+	"MarketplaceURL": "",
+	"SupportURL": "",
+	"EnabledByDefault": false,
+	"CanContainContent": false,
+	"IsBetaVersion": false,
+	"IsExperimentalVersion": true,
+	"Installed": false,
+	"Modules": [
+		{
+			"Name": "Chimera",
+			"Type": "Runtime",
+			"LoadingPhase": "Default"
+		},
+		{
+			"Name": "ChimeraEditor",
+			"Type": "UncookedOnly",
+			"LoadingPhase": "PreDefault"
+		}
+	],
+
+	"Plugins": [
+		{
+			"Name": "PoseSearch",
+			"Enabled": true
+		}
+	]
+}
\ No newline at end of file
diff --git a/Engine/Plugins/Experimental/Animation/Chimera/Source/Editor/ChimeraEditor.Build.cs b/Engine/Plugins/Experimental/Animation/Chimera/Source/Editor/ChimeraEditor.Build.cs
new file mode 100644
--- /dev/null
+++ b/Engine/Plugins/Experimental/Animation/Chimera/Source/Editor/ChimeraEditor.Build.cs
@@ -0,0 +1,38 @@
+// Copyright Epic Games, Inc. All Rights Reserved.
+
+namespace UnrealBuildTool.Rules
+{
+	public class ChimeraEditor : ModuleRules
+	{
+		public ChimeraEditor(ReadOnlyTargetRules Target) : base(Target)
+		{
+			PublicDependencyModuleNames.AddRange(
+			new string[]
+			{
+                "AnimGraph",
+				"AnimGraphRuntime",
+				"BlendStack",
+				"Core",
+                "CoreUObject",
+				"Chimera",
+                "Engine",
+				"PoseSearchEditor"
+			});
+
+            PrivateDependencyModuleNames.AddRange(
+			new string[]
+			{
+				"AdvancedPreviewScene",
+				"AssetDefinition",
+				"BlendStackEditor",
+				"BlueprintGraph",
+				"EditorFramework",
+				"PoseSearch",
+				"Slate",
+				"SlateCore",
+				"ToolWidgets",						
+				"UnrealEd",
+			});
+        }
+	}
+}
diff --git a/Engine/Plugins/Experimental/Animation/Chimera/Source/Editor/Private/AnimGraphNode_Chimera.cpp b/Engine/Plugins/Experimental/Animation/Chimera/Source/Editor/Private/AnimGraphNode_Chimera.cpp
new file mode 100644
--- /dev/null
+++ b/Engine/Plugins/Experimental/Animation/Chimera/Source/Editor/Private/AnimGraphNode_Chimera.cpp
@@ -0,0 +1,26 @@
+// Copyright Epic Games, Inc. All Rights Reserved.
+
+#include "AnimGraphNode_Chimera.h"
+
+#define LOCTEXT_NAMESPACE "Chimera"
+
+UAnimGraphNode_Chimera::UAnimGraphNode_Chimera(const FObjectInitializer& ObjectInitializer) : Super(ObjectInitializer)
+{
+}
+
+FText UAnimGraphNode_Chimera::GetTooltipText() const
+{
+	return LOCTEXT("ChimeraTooltip", "Character Interaction Matched Node");
+}
+
+FText UAnimGraphNode_Chimera::GetNodeTitle(ENodeTitleType::Type TitleType) const
+{
+	return LOCTEXT("Chimera", "Chimera");
+}
+
+FLinearColor UAnimGraphNode_Chimera::GetNodeTitleColor() const
+{
+	return FLinearColor(FColor(6, 9, 53));
+}
+
+#undef LOCTEXT_NAMESPACE
diff --git a/Engine/Plugins/Experimental/Animation/Chimera/Source/Editor/Private/ChimeraAssetDefinitions.cpp b/Engine/Plugins/Experimental/Animation/Chimera/Source/Editor/Private/ChimeraAssetDefinitions.cpp
new file mode 100644
--- /dev/null
+++ b/Engine/Plugins/Experimental/Animation/Chimera/Source/Editor/Private/ChimeraAssetDefinitions.cpp
@@ -0,0 +1,45 @@
+// Copyright Epic Games, Inc. All Rights Reserved.
+
+#include "ChimeraAssetDefinitions.h"
+#include "ChimeraAssetEditor.h"
+#include "ThumbnailRendering/SceneThumbnailInfo.h"
+
+namespace UE::Chimera
+{
+
+FLinearColor GetAssetColor()
+{
+	static const FLinearColor AssetColor(FColor(6, 9, 53));
+	return AssetColor;
+}
+
+TConstArrayView<FAssetCategoryPath> GetAssetCategories()
+{
+	static const auto Categories = { FAssetCategoryPath(EAssetCategoryPaths::Animation, NSLOCTEXT("ChimeraAssetDefinition", "ChimeraAssetDefinitionMenu", "Chimera")) };
+	return Categories;
+}
+
+UThumbnailInfo* LoadThumbnailInfo(const FAssetData & InAssetData)
+{
+	return UE::Editor::FindOrCreateThumbnailInfo(InAssetData.GetAsset(), USceneThumbnailInfo::StaticClass());
+}
+
+} // namespace UE::Chimera
+
+EAssetCommandResult UAssetDefinition_ChimeraAsset::OpenAssets(const FAssetOpenArgs& OpenArgs) const
+{
+	using namespace UE::Chimera;
+
+	TArray<UChimeraAsset*> Objects = OpenArgs.LoadObjects<UChimeraAsset>();
+	const EToolkitMode::Type Mode = OpenArgs.ToolkitHost.IsValid() ? EToolkitMode::WorldCentric : EToolkitMode::Standalone;
+	for (UChimeraAsset* ChimeraAsset : Objects)
+	{
+		if (ChimeraAsset)
+		{
+			TSharedRef<FChimeraAssetEditor> NewEditor(new FChimeraAssetEditor());
+			NewEditor->InitAssetEditor(Mode, OpenArgs.ToolkitHost, ChimeraAsset);
+		}
+	}
+	
+	return EAssetCommandResult::Handled;
+}
\ No newline at end of file
diff --git a/Engine/Plugins/Experimental/Animation/Chimera/Source/Editor/Private/ChimeraAssetDefinitions.h b/Engine/Plugins/Experimental/Animation/Chimera/Source/Editor/Private/ChimeraAssetDefinitions.h
new file mode 100644
--- /dev/null
+++ b/Engine/Plugins/Experimental/Animation/Chimera/Source/Editor/Private/ChimeraAssetDefinitions.h
@@ -0,0 +1,30 @@
+// Copyright Epic Games, Inc. All Rights Reserved.
+
+#pragma once
+
+#include "AssetDefinitionDefault.h"
+#include "Chimera/ChimeraAsset.h"
+#include "ChimeraAssetDefinitions.generated.h"
+
+namespace UE::Chimera
+{
+	FLinearColor GetAssetColor();
+	TConstArrayView<FAssetCategoryPath> GetAssetCategories();
+	UThumbnailInfo* LoadThumbnailInfo(const FAssetData& InAssetData);
+} // namespace UE::Chimera
+
+UCLASS()
+class UAssetDefinition_ChimeraAsset : public UAssetDefinitionDefault
+{
+	GENERATED_BODY()
+
+public:
+	virtual FLinearColor GetAssetColor() const override { return UE::Chimera::GetAssetColor(); }
+	virtual TConstArrayView<FAssetCategoryPath> GetAssetCategories() const override { return UE::Chimera::GetAssetCategories(); }
+	virtual UThumbnailInfo* LoadThumbnailInfo(const FAssetData& InAssetData) const override { return UE::Chimera::LoadThumbnailInfo(InAssetData); }
+
+	virtual FText GetAssetDisplayName() const override { return NSLOCTEXT("ChimeraAssetDefinition", "DisplayName_ChimeraAsset", "Chimera Asset"); }
+	virtual TSoftClassPtr<UObject> GetAssetClass() const override { return UChimeraAsset::StaticClass(); }
+	virtual EAssetCommandResult OpenAssets(const FAssetOpenArgs& OpenArgs) const override;
+};
+
diff --git a/Engine/Plugins/Experimental/Animation/Chimera/Source/Editor/Private/ChimeraAssetEditor.cpp b/Engine/Plugins/Experimental/Animation/Chimera/Source/Editor/Private/ChimeraAssetEditor.cpp
new file mode 100644
--- /dev/null
+++ b/Engine/Plugins/Experimental/Animation/Chimera/Source/Editor/Private/ChimeraAssetEditor.cpp
@@ -0,0 +1,896 @@
+// Copyright Epic Games, Inc. All Rights Reserved.
+
+#include "ChimeraAssetEditor.h"
+#include "Animation/DebugSkelMeshComponent.h"
+#include "AnimPreviewInstance.h"
+#include "AssetEditorModeManager.h"
+#include "Chimera/ChimeraAsset.h"
+#include "ChimeraEditorDefines.h"
+#include "Components/StaticMeshComponent.h"
+#include "Framework/Application/SlateApplication.h"
+#include "EngineUtils.h"
+#include "GameFramework/WorldSettings.h"
+#include "Modules/ModuleManager.h"
+#include "PreviewProfileController.h"
+#include "PropertyEditorModule.h"
+#include "Styling/AppStyle.h"
+#include "SSimpleTimeSlider.h"
+#include "UnrealWidget.h"
+#include "Viewports.h"
+#include "Widgets/Docking/SDockTab.h"
+#include "Widgets/Images/SImage.h"
+#include "Widgets/Input/SButton.h"
+
+#define LOCTEXT_NAMESPACE "ChimeraAssetEditor"
+
+namespace UE::Chimera
+{
+
+constexpr float StepDeltaTime = 1.f / 30.f;
+
+/////////////////////////////////////////////////
+// class FChimeraAssetEdMode
+bool FChimeraAssetPreviewActor::SpawnPreviewActor(UWorld* World, const UChimeraAsset* ChimeraAsset, const UE::PoseSearch::FRole& Role)
+{
+	UAnimationAsset* PreviewAsset = ChimeraAsset->GetAnimationAsset(Role);
+	if (!PreviewAsset)
+	{
+		return false;
+	}
+
+	ActorRole = Role;
+
+	FActorSpawnParameters Params;
+	Params.SpawnCollisionHandlingOverride = ESpawnActorCollisionHandlingMethod::AdjustIfPossibleButAlwaysSpawn;
+
+	for (int32 ActorIndex = 0; ActorIndex < NumActors; ++ActorIndex)
+	{
+		UE::PoseSearch::FAnimationAssetSampler& Sampler = Samplers[ActorIndex];
+
+		const FTransform Origin = ChimeraAsset->GetDebugWarpOrigin(GetRole(), ActorIndex == DebugActor);
+		Sampler.Init(PreviewAsset, Origin, BlendParameters);
+
+		const FTransform ActorTransform = Sampler.ExtractRootTransform(CurrentTime);
+
+		TWeakObjectPtr<AActor>& ActorPtr = ActorPtrs[ActorIndex];
+
+		ActorPtr = World->SpawnActor<AActor>(AActor::StaticClass(), FTransform::Identity, Params);
+		ActorPtr->SetFlags(RF_Transient);
+
+		UDebugSkelMeshComponent* Mesh = NewObject<UDebugSkelMeshComponent>(ActorPtr.Get());
+		Mesh->RegisterComponentWithWorld(World);
+
+		UAnimPreviewInstance* AnimInstance = NewObject<UAnimPreviewInstance>(Mesh);
+		Mesh->PreviewInstance = AnimInstance;
+		AnimInstance->InitializeAnimation();
+
+		USkeleton* Skeleton = PreviewAsset->GetSkeleton();
+		Mesh->SetSkeletalMesh(Skeleton->GetPreviewMesh(true));
+		Mesh->EnablePreview(true, PreviewAsset);
+
+		AnimInstance->SetAnimationAsset(PreviewAsset, ChimeraAsset->IsLooping(), 0.f);
+		AnimInstance->SetBlendSpacePosition(BlendParameters);
+
+		AnimInstance->PlayAnim(ChimeraAsset->IsLooping(), 0.f);
+		if (!ActorPtr->GetRootComponent())
+		{
+			ActorPtr->SetRootComponent(Mesh);
+		}
+
+		AnimInstance->SetPosition(Sampler.ToNormalizedTime(CurrentTime));
+		AnimInstance->SetPlayRate(0.f);
+		AnimInstance->SetBlendSpacePosition(BlendParameters);
+
+		ActorPtr->SetActorTransform(ActorTransform);
+
+		UE_LOG(LogChimeraEditor, Log, TEXT("Spawned preview Actor: %s"), *GetNameSafe(ActorPtr.Get()));
+	}
+	return true;
+}
+
+void FChimeraAssetPreviewActor::UpdatePreviewActor(const UChimeraAsset* ChimeraAsset, float PlayTime)
+{
+	if (UAnimationAsset* PreviewAsset = ChimeraAsset->GetAnimationAsset(GetRole()))
+	{
+		bool bPlayTimeUpdated = false;
+
+		float NewCurrentTime = 0.f;
+		FAnimationRuntime::AdvanceTime(false, PlayTime, NewCurrentTime, PreviewAsset->GetPlayLength());
+
+		if (!FMath::IsNearlyEqual(CurrentTime, NewCurrentTime))
+		{
+			CurrentTime = NewCurrentTime;
+			bPlayTimeUpdated = true;
+		}
+
+		for (int32 ActorIndex = 0; ActorIndex < NumActors; ++ActorIndex)
+		{
+			const FTransform Origin = ChimeraAsset->GetDebugWarpOrigin(GetRole(), ActorIndex == DebugActor);
+
+			bool bSamplerReinitialized = false;
+			UE::PoseSearch::FAnimationAssetSampler& Sampler = Samplers[ActorIndex];
+			if (PreviewAsset != Sampler.GetAsset() || !Origin.Equals(Sampler.GetRootTransformOrigin()))
+			{
+				// reinitializing the Sampler if the PreviewAsset or the origin transform changed
+				Sampler.Init(PreviewAsset, Origin, BlendParameters);
+				bSamplerReinitialized = true;
+			}
+
+			TWeakObjectPtr<AActor>& ActorPtr = ActorPtrs[ActorIndex];
+			if (ActorPtr != nullptr)
+			{
+				if (UDebugSkelMeshComponent* Mesh = Cast<UDebugSkelMeshComponent>(ActorPtr->GetRootComponent()))
+				{
+					if (UAnimPreviewInstance* AnimInstance = Mesh->PreviewInstance.Get())
+					{
+						bool bPreviewAssetChanged = false;
+						if (AnimInstance->GetAnimationAsset() != PreviewAsset)
+						{
+							AnimInstance->SetAnimationAsset(PreviewAsset, ChimeraAsset->IsLooping(), 0.f);
+							bPreviewAssetChanged = true;
+						}
+
+						if (bPlayTimeUpdated || bSamplerReinitialized || bPreviewAssetChanged)
+						{
+							// SetPosition is in [0..1] range for blendspaces
+							AnimInstance->SetPosition(Sampler.ToNormalizedTime(CurrentTime));
+							AnimInstance->SetPlayRate(0.f);
+							AnimInstance->SetBlendSpacePosition(BlendParameters);
+
+							const FTransform ActorTransform = Sampler.ExtractRootTransform(CurrentTime);
+							ActorPtr->SetActorTransform(ActorTransform);
+						}
+					}
+				}
+			}
+		}
+	}
+}
+
+void FChimeraAssetPreviewActor::Destroy()
+{
+	for (TWeakObjectPtr<AActor>& ActorPtr : ActorPtrs)
+	{
+		if (ActorPtr != nullptr)
+		{
+			ActorPtr->Destroy();
+		}
+		ActorPtr = nullptr;
+	}
+}
+
+
+UAnimPreviewInstance* FChimeraAssetPreviewActor::GetAnimPreviewInstance()
+{
+	if (const AActor* Actor = ActorPtrs[DebugActor].Get())
+	{
+		if (const UDebugSkelMeshComponent* Mesh = Cast<UDebugSkelMeshComponent>(Actor->GetRootComponent()))
+		{
+			return Mesh->PreviewInstance.Get();
+		}
+	}
+	
+	return nullptr;
+}
+
+FTransform FChimeraAssetPreviewActor::GetDebugActorTransformFromSampler() const
+{
+	const UE::PoseSearch::FAnimationAssetSampler& Sampler = Samplers[DebugActor];
+	const FTransform ActorTransform = Sampler.ExtractRootTransform(CurrentTime);
+	return ActorTransform;
+}
+
+void FChimeraAssetPreviewActor::ForceDebugActorTransform(const FTransform& ActorTransform)
+{
+	TWeakObjectPtr<AActor>& ActorPtr = ActorPtrs[DebugActor];
+	if (ActorPtr != nullptr)
+	{
+		ActorPtr->SetActorTransform(ActorTransform);
+	}
+}
+
+/////////////////////////////////////////////////
+// class FChimeraAssetViewModel
+void FChimeraAssetViewModel::AddReferencedObjects(FReferenceCollector& Collector)
+{
+	Collector.AddReferencedObject(ChimeraAssetPtr);
+}
+
+void FChimeraAssetViewModel::Initialize(UChimeraAsset* ChimeraAsset, const TSharedRef<FChimeraAssetPreviewScene>& PreviewScene)
+{
+	ChimeraAssetPtr = ChimeraAsset;
+	PreviewScenePtr = PreviewScene;
+}
+
+void FChimeraAssetViewModel::PreviewBackwardEnd()
+{
+	SetPlayTime(0.f, false);
+}
+
+void FChimeraAssetViewModel::PreviewBackwardStep()
+{
+	if (const UChimeraAsset* ChimeraAsset = GetChimeraAsset())
+	{
+		const float NewPlayTime = FMath::Clamp(PlayTime - StepDeltaTime, 0.f, ChimeraAsset->GetPlayLength());
+		SetPlayTime(NewPlayTime, false);
+	}
+}
+
+void FChimeraAssetViewModel::PreviewBackward()
+{
+	DeltaTimeMultiplier = -1.f;
+}
+
+void FChimeraAssetViewModel::PreviewPause()
+{
+	DeltaTimeMultiplier = 0.f;
+}
+
+void FChimeraAssetViewModel::PreviewForward()
+{
+	DeltaTimeMultiplier = 1.f;
+}
+
+void FChimeraAssetViewModel::PreviewForwardStep()
+{
+	if (const UChimeraAsset* ChimeraAsset = GetChimeraAsset())
+	{
+		const float NewPlayTime = FMath::Clamp(PlayTime + StepDeltaTime, 0.f, ChimeraAsset->GetPlayLength());
+		SetPlayTime(NewPlayTime, false);
+	}
+}
+
+void FChimeraAssetViewModel::PreviewForwardEnd()
+{
+	if (const UChimeraAsset* ChimeraAsset = GetChimeraAsset())
+	{
+		SetPlayTime(ChimeraAsset->GetPlayLength(), false);
+	}
+}
+
+UWorld* FChimeraAssetViewModel::GetWorld()
+{
+	check(PreviewScenePtr.IsValid());
+	return PreviewScenePtr.Pin()->GetWorld();
+}
+
+void FChimeraAssetViewModel::Tick(float DeltaSeconds)
+{
+	const UChimeraAsset* ChimeraAsset = GetChimeraAsset();
+	if (!ChimeraAsset)
+	{
+		RemovePreviewActors();
+		return;
+	}
+
+	PlayTime += DeltaSeconds * DeltaTimeMultiplier;
+	PlayTime = FMath::Clamp(PlayTime, 0.f, ChimeraAsset->GetPlayLength());
+
+	UE::PoseSearch::FRoleToIndex ChimeraAssetRoleToIndex;
+	for (int32 RoleIndex = 0; RoleIndex < ChimeraAsset->GetNumRoles(); ++RoleIndex)
+	{
+		ChimeraAssetRoleToIndex.Add(ChimeraAsset->GetRole(RoleIndex)) = RoleIndex;
+	}
+
+	// iterating backwards because of the possible RemoveAtSwap 
+	UE::PoseSearch::FRoleToIndex PreviewActorsRoleToIndex;
+	for (int32 ActorIndex = PreviewActors.Num() - 1; ActorIndex >= 0; --ActorIndex)
+	{
+		FChimeraAssetPreviewActor& PreviewActor = PreviewActors[ActorIndex];
+		if (!ChimeraAssetRoleToIndex.Find(PreviewActor.GetRole()))
+		{
+			PreviewActor.Destroy();
+			PreviewActors.RemoveAtSwap(ActorIndex, EAllowShrinking::No);
+		}
+		else
+		{
+			PreviewActorsRoleToIndex.Add(PreviewActor.GetRole());		
+		}
+	}
+
+	if (PreviewActors.Num() != ChimeraAssetRoleToIndex.Num())
+	{
+		for (int32 RoleIndex = 0; RoleIndex < ChimeraAsset->GetNumRoles(); ++RoleIndex)
+		{
+			if (!PreviewActorsRoleToIndex.Find(ChimeraAsset->GetRole(RoleIndex)))
+			{
+				FChimeraAssetPreviewActor PreviewActor;
+				if (PreviewActor.SpawnPreviewActor(GetWorld(), ChimeraAsset, ChimeraAsset->GetRole(RoleIndex)))
+				{
+					PreviewActors.Add(PreviewActor);
+				}
+			}
+		}
+	}
+
+	for (FChimeraAssetPreviewActor& PreviewActor : PreviewActors)
+	{
+		PreviewActor.UpdatePreviewActor(ChimeraAsset, PlayTime);
+	}
+	
+	UE::PoseSearch::FRoleToIndex PreviewActorRoleToIndex;
+	PreviewActorRoleToIndex.Reserve(PreviewActors.Num());
+	for (int32 PreviewActorIndex = 0; PreviewActorIndex < PreviewActors.Num(); ++PreviewActorIndex)
+	{
+		PreviewActorRoleToIndex.Add(PreviewActors[PreviewActorIndex].GetRole()) = PreviewActorIndex;
+	}
+
+	// testing CalculateWarpTransforms
+#if WITH_EDITORONLY_DATA
+	if (ChimeraAsset->bEnableDebugWarp && PreviewActors.Num() == ChimeraAsset->GetNumRoles())
+	{
+		const UWorld* DebugDrawWorld = nullptr;
+		TArray<FTransform, TInlineAllocator<UE::PoseSearch::PreallocatedRolesNum>> ActorTransforms;
+		ActorTransforms.SetNum(ChimeraAsset->GetNumRoles());
+		for (int32 ChimeraAssetRoleIndex = 0; ChimeraAssetRoleIndex < ChimeraAsset->GetNumRoles(); ++ChimeraAssetRoleIndex)
+		{
+			const UE::PoseSearch::FRole& ChimeraAssetRole = ChimeraAsset->GetRole(ChimeraAssetRoleIndex);
+			const int32 PreviewActorIndex = PreviewActorRoleToIndex[ChimeraAssetRole];
+
+			FChimeraAssetPreviewActor& PreviewActor = PreviewActors[PreviewActorIndex];
+			check(PreviewActor.GetRole() == ChimeraAssetRole);
+
+			ActorTransforms[ChimeraAssetRoleIndex] = PreviewActor.GetDebugActorTransformFromSampler();
+			if (!DebugDrawWorld)
+			{
+				if (UAnimPreviewInstance* AnimPreviewInstance = PreviewActor.GetAnimPreviewInstance())
+				{
+					DebugDrawWorld = AnimPreviewInstance->GetWorld();
+				}
+			}
+		}
+
+		TArray<FTransform, TInlineAllocator<UE::PoseSearch::PreallocatedRolesNum>> FullAlignedActorTransforms;
+		FullAlignedActorTransforms.SetNum(ChimeraAsset->GetNumRoles());
+		ChimeraAsset->CalculateWarpTransforms(PlayTime, ActorTransforms, FullAlignedActorTransforms, DebugDrawWorld);
+
+		for (int32 ChimeraAssetRoleIndex = 0; ChimeraAssetRoleIndex < ChimeraAsset->GetNumRoles(); ++ChimeraAssetRoleIndex)
+		{
+			const UE::PoseSearch::FRole& ChimeraAssetRole = ChimeraAsset->GetRole(ChimeraAssetRoleIndex);
+			const int32 PreviewActorIndex = PreviewActorRoleToIndex[ChimeraAssetRole];
+
+			FChimeraAssetPreviewActor& PreviewActor = PreviewActors[PreviewActorIndex];
+			check(PreviewActor.GetRole() == ChimeraAssetRole);
+			FTransform DebugActorTransform;
+			DebugActorTransform.Blend(ActorTransforms[ChimeraAssetRoleIndex], FullAlignedActorTransforms[ChimeraAssetRoleIndex], ChimeraAsset->DebugWarpAmount);
+			PreviewActor.ForceDebugActorTransform(DebugActorTransform);
+		}
+	}
+#endif // WITH_EDITORONLY_DATA
+}
+
+void FChimeraAssetViewModel::RemovePreviewActors()
+{
+	PlayTime = 0.f;
+	DeltaTimeMultiplier = 1.f;
+
+	for (FChimeraAssetPreviewActor& PreviewActor : PreviewActors)
+	{
+		PreviewActor.Destroy();
+	}
+
+	PreviewActors.Reset();
+}
+
+TRange<double> FChimeraAssetViewModel::GetPreviewPlayRange() const
+{
+	constexpr double ViewRangeSlack = 0.2;
+	if (const UChimeraAsset* ChimeraAsset = GetChimeraAsset())
+	{
+		return TRange<double>(-ViewRangeSlack, ChimeraAsset->GetPlayLength() + ViewRangeSlack);
+	}
+	
+	return TRange<double>(-ViewRangeSlack, ViewRangeSlack);
+}
+
+void FChimeraAssetViewModel::SetPlayTime(float NewPlayTime, bool bInTickPlayTime)
+{
+	if (const UChimeraAsset* ChimeraAsset = GetChimeraAsset())
+	{
+		NewPlayTime = FMath::Clamp(NewPlayTime, 0.f, ChimeraAsset->GetPlayLength());
+		DeltaTimeMultiplier = bInTickPlayTime ? DeltaTimeMultiplier : 0.f;
+
+		if (!FMath::IsNearlyEqual(PlayTime, NewPlayTime))
+		{
+			PlayTime = NewPlayTime;
+
+			for (FChimeraAssetPreviewActor& PreviewActor : PreviewActors)
+			{
+				PreviewActor.UpdatePreviewActor(ChimeraAsset, PlayTime);
+			}
+		}
+	}
+}
+
+void FChimeraAssetViewModel::SetPreviewProperties(float AnimAssetTime, const FVector& AnimAssetBlendParameters, bool bAnimAssetPlaying)
+{
+	// @todo: add support for blend spaces (pass AnimAssetBlendParameters as input)
+	SetPlayTime(AnimAssetTime, bAnimAssetPlaying);
+}
+
+/////////////////////////////////////////////////
+// class FChimeraAssetEdMode
+const FEditorModeID FChimeraAssetEdMode::EdModeId = TEXT("ChimeraAssetEdMode");
+
+void FChimeraAssetEdMode::Tick(FEditorViewportClient* ViewportClient, float DeltaTime)
+{
+	FEdMode::Tick(ViewportClient, DeltaTime);
+
+	if (FChimeraAssetViewportClient* ChimeraAssetViewportClient = static_cast<FChimeraAssetViewportClient*>(ViewportClient))
+	{
+		// ensure we redraw even if PIE is active
+		ChimeraAssetViewportClient->Invalidate();
+
+		if (!ViewModel)
+		{
+			ViewModel = ChimeraAssetViewportClient->GetAssetEditor()->GetViewModel();
+		}
+	}
+
+	if (ViewModel)
+	{
+		ViewModel->Tick(DeltaTime);
+	}
+}
+
+void FChimeraAssetEdMode::Render(const FSceneView* View, FViewport* Viewport, FPrimitiveDrawInterface* PDI)
+{
+	FEdMode::Render(View, Viewport, PDI);
+}
+
+bool FChimeraAssetEdMode::AllowWidgetMove()
+{
+	return FEdMode::ShouldDrawWidget();
+}
+
+bool FChimeraAssetEdMode::ShouldDrawWidget() const
+{
+	return FEdMode::ShouldDrawWidget();
+}
+
+bool FChimeraAssetEdMode::GetCustomInputCoordinateSystem(FMatrix& InMatrix, void* InData)
+{
+	return FEdMode::GetCustomDrawingCoordinateSystem(InMatrix, InData);
+}
+
+/////////////////////////////////////////////////
+// class FChimeraAssetViewportClient
+FChimeraAssetViewportClient::FChimeraAssetViewportClient(
+	const TSharedRef<FChimeraAssetPreviewScene>& InPreviewScene,
+	const TSharedRef<SChimeraAssetViewport>& InViewport,
+	const TSharedRef<FChimeraAssetEditor>& InAssetEditor)
+	: FEditorViewportClient(nullptr, &InPreviewScene.Get(), StaticCastSharedRef<SEditorViewport>(InViewport))
+	, PreviewScenePtr(InPreviewScene)
+	, AssetEditorPtr(InAssetEditor)
+{
+	Widget->SetUsesEditorModeTools(ModeTools.Get());
+	StaticCastSharedPtr<FAssetEditorModeManager>(ModeTools)->SetPreviewScene(&InPreviewScene.Get());
+	ModeTools->SetDefaultMode(FChimeraAssetEdMode::EdModeId);
+
+	SetRealtime(true);
+
+	SetWidgetCoordSystemSpace(COORD_Local);
+	ModeTools->SetWidgetMode(UE::Widget::WM_Translate);
+}
+
+void FChimeraAssetViewportClient::TrackingStarted(
+	const struct FInputEventState& InInputState,
+	bool bIsDraggingWidget,
+	bool bNudge)
+{
+	ModeTools->StartTracking(this, Viewport);
+}
+
+void FChimeraAssetViewportClient::TrackingStopped()
+{
+	ModeTools->EndTracking(this, Viewport);
+	Invalidate();
+}
+
+void FChimeraAssetViewportClient::Draw(const FSceneView* View, FPrimitiveDrawInterface* PDI)
+{
+	FEditorViewportClient::Draw(View, PDI);
+}
+
+/////////////////////////////////////////////////
+// class FChimeraAssetPreviewScene
+FChimeraAssetPreviewScene::FChimeraAssetPreviewScene(ConstructionValues CVs, const TSharedRef<FChimeraAssetEditor>& Editor)
+: FAdvancedPreviewScene(CVs)
+, EditorPtr(Editor)
+{
+	// Disable killing actors outside of the world
+	AWorldSettings* WorldSettings = GetWorld()->GetWorldSettings(true);
+	WorldSettings->bEnableWorldBoundsChecks = false;
+
+	// Spawn an owner for FloorMeshComponent so CharacterMovementComponent can detect it as a valid floor and slide 
+	// along it
+	{
+		AActor* FloorActor = GetWorld()->SpawnActor<AActor>(AActor::StaticClass(), FTransform());
+		check(FloorActor);
+
+		static const FString NewName = FString(TEXT("FloorComponent"));
+		FloorMeshComponent->Rename(*NewName, FloorActor);
+
+		FloorActor->SetRootComponent(FloorMeshComponent);
+	}
+}
+
+void FChimeraAssetPreviewScene::Tick(float InDeltaTime)
+{
+	FAdvancedPreviewScene::Tick(InDeltaTime);
+
+	// Trigger Begin Play in this preview world.
+	// This is needed for the CharacterMovementComponent to be able to switch to falling mode. 
+	// See: UCharacterMovementComponent::StartFalling
+	if (PreviewWorld && !PreviewWorld->GetBegunPlay())
+	{
+		for (FActorIterator It(PreviewWorld); It; ++It)
+		{
+			It->DispatchBeginPlay();
+		}
+
+		PreviewWorld->SetBegunPlay(true);
+	}
+
+	GetWorld()->Tick(LEVELTICK_All, InDeltaTime);
+}
+
+/////////////////////////////////////////////////
+// class SChimeraAssetViewport
+void SChimeraAssetViewport::Construct(
+	const FArguments& InArgs,
+	const FChimeraAssetPreviewRequiredArgs& InRequiredArgs)
+{
+	PreviewScenePtr = InRequiredArgs.PreviewScene;
+	AssetEditorPtr = InRequiredArgs.AssetEditor;
+
+	SEditorViewport::Construct(
+		SEditorViewport::FArguments()
+		.IsEnabled(FSlateApplication::Get().GetNormalExecutionAttribute())
+		.AddMetaData<FTagMetaData>(TEXT("AnimationTools.Viewport"))
+	);
+}
+
+void SChimeraAssetViewport::BindCommands()
+{
+	SEditorViewport::BindCommands();
+}
+
+TSharedRef<FEditorViewportClient> SChimeraAssetViewport::MakeEditorViewportClient()
+{
+	ViewportClient = MakeShared<FChimeraAssetViewportClient>(
+		PreviewScenePtr.Pin().ToSharedRef(),
+		SharedThis(this),
+		AssetEditorPtr.Pin().ToSharedRef());
+	ViewportClient->ViewportType = LVT_Perspective;
+	ViewportClient->bSetListenerPosition = false;
+	ViewportClient->SetViewLocation(EditorViewportDefs::DefaultPerspectiveViewLocation);
+	ViewportClient->SetViewRotation(EditorViewportDefs::DefaultPerspectiveViewRotation);
+
+	return ViewportClient.ToSharedRef();
+}
+
+TSharedPtr<SWidget> SChimeraAssetViewport::MakeViewportToolbar()
+{
+	return SAssignNew(ViewportToolbar, SCommonEditorViewportToolbarBase, SharedThis(this));
+}
+
+TSharedRef<SEditorViewport> SChimeraAssetViewport::GetViewportWidget()
+{
+	return SharedThis(this);
+}
+
+TSharedPtr<FExtender> SChimeraAssetViewport::GetExtenders() const
+{
+	TSharedPtr<FExtender> Result(MakeShareable(new FExtender));
+	return Result;
+}
+
+void SChimeraAssetViewport::OnFloatingButtonClicked()
+{
+}
+
+/////////////////////////////////////////////////
+// class SChimeraAssetPreview
+void SChimeraAssetPreview::Construct(const FArguments& InArgs, const FChimeraAssetPreviewRequiredArgs& InRequiredArgs)
+{
+	SliderColor = InArgs._SliderColor;
+	SliderScrubTime = InArgs._SliderScrubTime;
+	SliderViewRange = InArgs._SliderViewRange;
+	OnSliderScrubPositionChanged = InArgs._OnSliderScrubPositionChanged;
+
+	OnBackwardEnd = InArgs._OnBackwardEnd;
+	OnBackwardStep = InArgs._OnBackwardStep;
+	OnBackward = InArgs._OnBackward;
+	OnPause = InArgs._OnPause;
+	OnForward = InArgs._OnForward;
+	OnForwardStep = InArgs._OnForwardStep;
+	OnForwardEnd = InArgs._OnForwardEnd;
+
+	FSlimHorizontalToolBarBuilder ToolBarBuilder(
+		TSharedPtr<const FUICommandList>(), 
+		FMultiBoxCustomization::None, 
+		nullptr, true);
+
+	auto AddToolBarButton = [&ToolBarBuilder](FName ButtonImageName, FOnButtonClickedEvent& OnClicked)
+		{
+			ToolBarBuilder.AddToolBarWidget(
+				SNew(SButton)
+				.ButtonStyle(FAppStyle::Get(), "Animation.PlayControlsButton")
+				.OnClicked_Lambda([&OnClicked]()
+					{
+						if (OnClicked.IsBound())
+						{
+							OnClicked.Execute();
+							return FReply::Handled();
+						}
+						return FReply::Unhandled();
+					})
+				[
+					SNew(SImage)
+						.ColorAndOpacity(FSlateColor::UseSubduedForeground())
+						.Image(FAppStyle::Get().GetBrush(ButtonImageName))
+				]);
+		};
+
+	//ToolBarBuilder.SetStyle(&FAppStyle::Get(), "PaletteToolBar");
+	ToolBarBuilder.BeginSection("Preview");
+	{
+		AddToolBarButton("Animation.Backward_End", OnBackwardEnd);
+		AddToolBarButton("Animation.Backward_Step", OnBackwardStep);
+		AddToolBarButton("Animation.Backward", OnBackward);
+		AddToolBarButton("Animation.Pause", OnPause);
+		AddToolBarButton("Animation.Forward", OnForward);
+		AddToolBarButton("Animation.Forward_Step", OnForwardStep);
+		AddToolBarButton("Animation.Forward_End", OnForwardEnd);
+	}
+
+	ChildSlot
+	[
+		SNew(SVerticalBox)
+		+SVerticalBox::Slot()
+		.FillHeight(1.0f)
+		[
+			SNew(SChimeraAssetViewport, InRequiredArgs)
+		]
+		+SVerticalBox::Slot()
+		.AutoHeight()
+		[
+			SNew(SHorizontalBox)
+			+SHorizontalBox::Slot()
+			.AutoWidth()
+			[
+				ToolBarBuilder.MakeWidget()
+			]
+			+SHorizontalBox::Slot()
+			.FillWidth(1.0f)
+			[
+				SNew(SSimpleTimeSlider)
+				.ClampRangeHighlightSize(0.15f)
+				.ClampRangeHighlightColor_Lambda([this]()
+					{
+						return SliderColor.Get();
+					})
+				.ScrubPosition_Lambda([this]()
+					{
+						return SliderScrubTime.Get();
+					})
+				.ViewRange_Lambda([this]()
+					{
+						return SliderViewRange.Get();
+					})
+				.ClampRange_Lambda([this]()
+					{
+						return SliderViewRange.Get();
+					})
+				.OnScrubPositionChanged_Lambda([this](double NewScrubTime, bool bIsScrubbing)
+					{
+						if (bIsScrubbing)
+						{
+							OnSliderScrubPositionChanged.ExecuteIfBound(NewScrubTime, bIsScrubbing);
+						}
+					})
+			]
+		]
+	];
+}
+
+/////////////////////////////////////////////////
+// class FChimeraAssetEditor
+const FName ChimeraAssetEditorAppName = FName(TEXT("ChimeraAssetEditorApp"));
+
+// Tab identifiers
+struct FChimeraAssetEditorTabs
+{
+	static const FName AssetDetailsID;
+	static const FName ViewportID;
+};
+const FName FChimeraAssetEditorTabs::AssetDetailsID(TEXT("ChimeraAssetEditorAssetDetailsTabID"));
+const FName FChimeraAssetEditorTabs::ViewportID(TEXT("ChimeraAssetEditorViewportTabID"));
+
+void FChimeraAssetEditor::InitAssetEditor(const EToolkitMode::Type Mode, const TSharedPtr<IToolkitHost>& InitToolkitHost, UChimeraAsset* ChimeraAsset)
+{
+	// Create Preview Scene
+	if (!PreviewScene.IsValid())
+	{
+		PreviewScene = MakeShareable(
+				new FChimeraAssetPreviewScene(
+					FPreviewScene::ConstructionValues()
+					.SetCreatePhysicsScene(false)
+					.SetTransactional(false)
+					.ForceUseMovementComponentInNonGameWorld(true),
+					StaticCastSharedRef<FChimeraAssetEditor>(AsShared())));
+
+		//Temporary fix for missing attached assets - MDW (Copied from FPersonaToolkit::CreatePreviewScene)
+		PreviewScene->GetWorld()->GetWorldSettings()->SetIsTemporarilyHiddenInEditor(false);
+	}
+
+	// Create view model
+	ViewModel = MakeShared<FChimeraAssetViewModel>();
+	ViewModel->Initialize(ChimeraAsset, PreviewScene.ToSharedRef());
+
+	// Create viewport widget
+	{
+		FChimeraAssetPreviewRequiredArgs PreviewArgs(
+			StaticCastSharedRef<FChimeraAssetEditor>(AsShared()),
+			PreviewScene.ToSharedRef());
+			
+		PreviewWidget = SNew(SChimeraAssetPreview, PreviewArgs)
+			.SliderColor(FLinearColor::Red)
+			.SliderScrubTime_Lambda([this]()
+				{
+					return ViewModel->GetPlayTime();
+				})
+			.SliderViewRange_Lambda([this]() 
+				{ 
+					return ViewModel->GetPreviewPlayRange();
+				})
+			.OnSliderScrubPositionChanged_Lambda([this](float NewScrubPosition, bool bScrubbing)
+				{
+					ViewModel->SetPlayTime(NewScrubPosition, !bScrubbing);
+				})
+			.OnBackwardEnd_Raw(this, &FChimeraAssetEditor::PreviewBackwardEnd)
+			.OnBackwardStep_Raw(this, &FChimeraAssetEditor::PreviewBackwardStep)
+			.OnBackward_Raw(this, &FChimeraAssetEditor::PreviewBackward)
+			.OnPause_Raw(this, &FChimeraAssetEditor::PreviewPause)
+			.OnForward_Raw(this, &FChimeraAssetEditor::PreviewForward)
+			.OnForwardStep_Raw(this, &FChimeraAssetEditor::PreviewForwardStep)
+			.OnForwardEnd_Raw(this, &FChimeraAssetEditor::PreviewForwardEnd);
+	}
+
+	// asset details widget
+	FDetailsViewArgs AssetDetailsArgs;
+	AssetDetailsArgs.bHideSelectionTip = true;
+	AssetDetailsArgs.NotifyHook = this;
+
+	FPropertyEditorModule& PropertyModule = FModuleManager::LoadModuleChecked<FPropertyEditorModule>("PropertyEditor");
+	EditingAssetWidget = PropertyModule.CreateDetailView(AssetDetailsArgs);
+	EditingAssetWidget->SetObject(ChimeraAsset);
+
+	// Define Editor Layout
+	const TSharedRef<FTabManager::FLayout> StandaloneDefaultLayout =
+	FTabManager::NewLayout("Standalone_ChimeraAssetDatabaseEditor_Layout_v0.02")
+		->AddArea
+		(
+			// Main application area
+			FTabManager::NewPrimaryArea()->SetOrientation(Orient_Horizontal)
+			->Split
+			(
+				FTabManager::NewStack()
+					->SetSizeCoefficient(0.4f)
+					->AddTab(FChimeraAssetEditorTabs::AssetDetailsID, ETabState::OpenedTab)
+					->SetHideTabWell(false)
+			)
+			->Split
+			(
+				FTabManager::NewStack()
+					->SetSizeCoefficient(0.6f)
+					->AddTab(FChimeraAssetEditorTabs::ViewportID, ETabState::OpenedTab)
+					->SetHideTabWell(false)
+			)
+		);
+
+	FAssetEditorToolkit::InitAssetEditor(
+		Mode,
+		InitToolkitHost,
+		ChimeraAssetEditorAppName,
+		StandaloneDefaultLayout,
+		true,
+		true,
+		ChimeraAsset,
+		false);
+
+	RegenerateMenusAndToolbars();
+}
+
+void FChimeraAssetEditor::RegisterTabSpawners(const TSharedRef<FTabManager>& InTabManager)
+{
+	WorkspaceMenuCategory = InTabManager->AddLocalWorkspaceMenuCategory(
+		LOCTEXT("WorkspaceMenu_ChimeraAssetEditor", "Chimera Asset Editor"));
+	TSharedRef<FWorkspaceItem> WorkspaceMenuCategoryRef = WorkspaceMenuCategory.ToSharedRef();
+
+	FAssetEditorToolkit::RegisterTabSpawners(InTabManager);
+
+	InTabManager->RegisterTabSpawner(
+		FChimeraAssetEditorTabs::ViewportID,
+		FOnSpawnTab::CreateSP(this, &FChimeraAssetEditor::SpawnTab_Viewport))
+		.SetDisplayName(LOCTEXT("ViewportTab", "Viewport"))
+		.SetGroup(WorkspaceMenuCategoryRef)
+		.SetIcon(FSlateIcon(FAppStyle::GetAppStyleSetName(), "GraphEditor.EventGraph_16x"));
+			
+	InTabManager->RegisterTabSpawner(
+		FChimeraAssetEditorTabs::AssetDetailsID,
+		FOnSpawnTab::CreateSP(this, &FChimeraAssetEditor::SpawnTab_AssetDetails))
+		.SetDisplayName(LOCTEXT("ChimeraAssetDetailsTab", "Chimera Asset Details"))
+		.SetGroup(WorkspaceMenuCategoryRef)
+		.SetIcon(FSlateIcon(FAppStyle::GetAppStyleSetName(), "LevelEditor.Tabs.Details"));
+}
+
+void FChimeraAssetEditor::UnregisterTabSpawners(const TSharedRef<FTabManager>& InTabManager)
+{
+	FAssetEditorToolkit::UnregisterTabSpawners(InTabManager);
+
+	InTabManager->UnregisterTabSpawner(FChimeraAssetEditorTabs::ViewportID);
+	InTabManager->UnregisterTabSpawner(FChimeraAssetEditorTabs::AssetDetailsID);
+}
+
+FName FChimeraAssetEditor::GetToolkitFName() const
+{
+	return FName("ChimeraAssetEditor");
+}
+
+FText FChimeraAssetEditor::GetBaseToolkitName() const
+{
+	return LOCTEXT("ChimeraAssetEditorAppLabel", "Chimera Asset Editor");
+}
+
+FText FChimeraAssetEditor::GetToolkitName() const
+{
+	FFormatNamedArguments Args;
+	Args.Add(TEXT("AssetName"), FText::FromString(GetNameSafe(GetChimeraAsset())));
+	return FText::Format(LOCTEXT("ChimeraAssetEditorToolkitName", "{AssetName}"), Args);
+}
+
+FLinearColor FChimeraAssetEditor::GetWorldCentricTabColorScale() const
+{
+	return FLinearColor::White;
+}
+
+FString FChimeraAssetEditor::GetWorldCentricTabPrefix() const
+{
+	return TEXT("ChimeraAssetEditor");
+}
+
+void FChimeraAssetEditor::SetPreviewProperties(float AnimAssetTime, const FVector& AnimAssetBlendParameters, bool bAnimAssetPlaying)
+{
+	ViewModel->SetPreviewProperties(AnimAssetTime, AnimAssetBlendParameters, bAnimAssetPlaying);
+}
+
+TSharedRef<SDockTab> FChimeraAssetEditor::SpawnTab_Viewport(const FSpawnTabArgs& Args)
+{
+	check(Args.GetTabId() == FChimeraAssetEditorTabs::ViewportID);
+
+	TSharedRef<SDockTab> SpawnedTab = SNew(SDockTab).Label(LOCTEXT("ViewportTab_Title", "Viewport"));
+
+	if (PreviewWidget.IsValid())
+	{
+		SpawnedTab->SetContent(PreviewWidget.ToSharedRef());
+	}
+
+	return SpawnedTab;
+}
+
+TSharedRef<SDockTab> FChimeraAssetEditor::SpawnTab_AssetDetails(const FSpawnTabArgs& Args)
+{
+	check(Args.GetTabId() == FChimeraAssetEditorTabs::AssetDetailsID);
+
+	return SNew(SDockTab)
+		.Label(LOCTEXT("ChimeraAsset_Details_Title", "Chimera Asset Details"))
+		[
+			EditingAssetWidget.ToSharedRef()
+		];
+}
+
+} // namespace UE::Chimera
+
+#undef LOCTEXT_NAMESPACE
diff --git a/Engine/Plugins/Experimental/Animation/Chimera/Source/Editor/Private/ChimeraAssetFactory.cpp b/Engine/Plugins/Experimental/Animation/Chimera/Source/Editor/Private/ChimeraAssetFactory.cpp
new file mode 100644
--- /dev/null
+++ b/Engine/Plugins/Experimental/Animation/Chimera/Source/Editor/Private/ChimeraAssetFactory.cpp
@@ -0,0 +1,26 @@
+// Copyright Epic Games, Inc. All Rights Reserved.
+
+#include "ChimeraAssetFactory.h"
+#include "Chimera/ChimeraAsset.h"
+
+#define LOCTEXT_NAMESPACE "ChimeraAssetFactory"
+
+UChimeraAssetFactory::UChimeraAssetFactory(const FObjectInitializer& ObjectInitializer)
+	: Super(ObjectInitializer)
+{
+	bCreateNew = true;
+	bEditAfterNew = true;
+	SupportedClass = UChimeraAsset::StaticClass();
+}
+
+UObject* UChimeraAssetFactory::FactoryCreateNew(UClass* Class, UObject* InParent, FName Name, EObjectFlags Flags, UObject* Context, FFeedbackContext* Warn)
+{
+	return NewObject<UChimeraAsset>(InParent, Class, Name, Flags);
+}
+
+FString UChimeraAssetFactory::GetDefaultNewAssetName() const
+{
+	return FString(TEXT("NewChimeraAsset"));
+}
+
+#undef LOCTEXT_NAMESPACE
\ No newline at end of file
diff --git a/Engine/Plugins/Experimental/Animation/Chimera/Source/Editor/Private/ChimeraEditorDefines.cpp b/Engine/Plugins/Experimental/Animation/Chimera/Source/Editor/Private/ChimeraEditorDefines.cpp
new file mode 100644
--- /dev/null
+++ b/Engine/Plugins/Experimental/Animation/Chimera/Source/Editor/Private/ChimeraEditorDefines.cpp
@@ -0,0 +1,5 @@
+// Copyright Epic Games, Inc. All Rights Reserved.
+
+#include "ChimeraEditorDefines.h"
+
+DEFINE_LOG_CATEGORY(LogChimeraEditor);
diff --git a/Engine/Plugins/Experimental/Animation/Chimera/Source/Editor/Private/ChimeraEditorModule.cpp b/Engine/Plugins/Experimental/Animation/Chimera/Source/Editor/Private/ChimeraEditorModule.cpp
new file mode 100644
--- /dev/null
+++ b/Engine/Plugins/Experimental/Animation/Chimera/Source/Editor/Private/ChimeraEditorModule.cpp
@@ -0,0 +1,39 @@
+// Copyright Epic Games, Inc. All Rights Reserved.
+
+#include "ChimeraAssetEditor.h"
+#include "EditorModeRegistry.h"
+#include "Modules/ModuleManager.h"
+#include "Modules/ModuleInterface.h"
+
+#define LOCTEXT_NAMESPACE "FChimeraEditorModule"
+
+namespace UE::Chimera
+{
+
+class FChimeraEditorModule : public IModuleInterface
+{
+private:
+	virtual void StartupModule() override;
+	virtual void ShutdownModule() override;
+};
+
+void FChimeraEditorModule::StartupModule()
+{
+	if (GIsEditor && !IsRunningCommandlet())
+	{
+		// Register Ed Mode used by ChimeraAsset
+		FEditorModeRegistry::Get().RegisterMode<FChimeraAssetEdMode>(FChimeraAssetEdMode::EdModeId, LOCTEXT("ChimeraAssetEdModeName", "ChimeraAsset"));
+	}
+}
+
+void FChimeraEditorModule::ShutdownModule()
+{
+	// Unregister Ed Mode
+	FEditorModeRegistry::Get().UnregisterMode(FChimeraAssetEdMode::EdModeId);
+}
+
+} // namespace UE::Chimera
+
+IMPLEMENT_MODULE(UE::Chimera::FChimeraEditorModule, ChimeraEditor)
+
+#undef LOCTEXT_NAMESPACE
diff --git a/Engine/Plugins/Experimental/Animation/Chimera/Source/Editor/Public/AnimGraphNode_Chimera.h b/Engine/Plugins/Experimental/Animation/Chimera/Source/Editor/Public/AnimGraphNode_Chimera.h
new file mode 100644
--- /dev/null
+++ b/Engine/Plugins/Experimental/Animation/Chimera/Source/Editor/Public/AnimGraphNode_Chimera.h
@@ -0,0 +1,24 @@
+// Copyright Epic Games, Inc. All Rights Reserved.
+
+#pragma once
+
+#include "AnimGraphNode_BlendStack.h"
+#include "Chimera/AnimNode_Chimera.h"
+#include "AnimGraphNode_Chimera.generated.h"
+
+UCLASS(Experimental)
+class CHIMERAEDITOR_API UAnimGraphNode_Chimera : public UAnimGraphNode_BlendStack_Base
+{
+	GENERATED_UCLASS_BODY()
+
+	UPROPERTY(EditAnywhere, Category = Settings)
+	FAnimNode_Chimera Node;
+
+public:
+	virtual FText GetNodeTitle(ENodeTitleType::Type TitleType) const override;
+	virtual FLinearColor GetNodeTitleColor() const override;
+	virtual FText GetTooltipText() const override;
+
+protected:
+	virtual FAnimNode_BlendStack_Standalone* GetBlendStackNode() const override { return (FAnimNode_BlendStack_Standalone*)(&Node); }
+};
diff --git a/Engine/Plugins/Experimental/Animation/Chimera/Source/Editor/Public/ChimeraAssetEditor.h b/Engine/Plugins/Experimental/Animation/Chimera/Source/Editor/Public/ChimeraAssetEditor.h
new file mode 100644
--- /dev/null
+++ b/Engine/Plugins/Experimental/Animation/Chimera/Source/Editor/Public/ChimeraAssetEditor.h
@@ -0,0 +1,262 @@
+// Copyright Epic Games, Inc. All Rights Reserved.
+
+#pragma once
+
+#include "AdvancedPreviewScene.h"
+#include "CoreMinimal.h"
+#include "EditorViewportClient.h"
+#include "EdMode.h"
+#include "IMultiAnimAssetEditor.h"
+#include "Misc/NotifyHook.h"
+#include "PoseSearch/PoseSearchAssetSampler.h"
+#include "PoseSearch/PoseSearchRole.h"
+#include "SCommonEditorViewportToolbarBase.h"
+#include "SEditorViewport.h"
+
+class UAnimPreviewInstance;
+class UChimeraAsset;
+class UDebugSkelMeshComponent;
+
+namespace UE::Chimera
+{
+
+class FChimeraAssetPreviewScene;
+class FChimeraAssetEditor;
+class SChimeraAssetPreview;
+class SChimeraAssetViewport;
+
+/////////////////////////////////////////////////
+// class struct FChimeraAssetPreviewActor
+struct FChimeraAssetPreviewActor
+{
+public:
+	bool SpawnPreviewActor(UWorld* World, const UChimeraAsset* ChimeraAsset, const UE::PoseSearch::FRole& Role);
+	void UpdatePreviewActor(const UChimeraAsset* ChimeraAsset, float PlayTime);
+	void Destroy();
+	UAnimPreviewInstance* GetAnimPreviewInstance();
+	const UE::PoseSearch::FRole& GetRole() const { return ActorRole; }
+	FTransform GetDebugActorTransformFromSampler() const;
+	void ForceDebugActorTransform(const FTransform& ActorTransform);
+
+private:
+	enum { PreviewActor, DebugActor, NumActors };
+	TStaticArray<TWeakObjectPtr<AActor>, NumActors> ActorPtrs;
+	TStaticArray<UE::PoseSearch::FAnimationAssetSampler, NumActors> Samplers; // used to keep track of the animation position / orientation
+
+	float CurrentTime = 0.f;
+	UE::PoseSearch::FRole ActorRole = UE::PoseSearch::DefaultRole;
+	
+	// @todo: add support for blend spaces
+	FVector BlendParameters = FVector::ZeroVector;
+};
+
+/////////////////////////////////////////////////
+// class FChimeraAssetViewModel
+class FChimeraAssetViewModel : public TSharedFromThis<FChimeraAssetViewModel>, public FGCObject
+{
+public:
+	virtual void AddReferencedObjects(FReferenceCollector& Collector) override;
+	virtual FString GetReferencerName() const override { return TEXT("FChimeraAssetViewModel"); }
+	void Initialize(UChimeraAsset* ChimeraAsset, const TSharedRef<FChimeraAssetPreviewScene>& PreviewScene);
+	void RemovePreviewActors();
+	void PreviewBackwardEnd();
+	void PreviewBackwardStep();
+	void PreviewBackward();
+	void PreviewPause();
+	void PreviewForward();
+	void PreviewForwardStep();
+	void PreviewForwardEnd();
+	const UChimeraAsset* GetChimeraAsset() const { return ChimeraAssetPtr.Get(); }
+	void Tick(float DeltaSeconds);
+	const TArray<FChimeraAssetPreviewActor>& GetPreviewActors() const { return PreviewActors; }
+	TArray<FChimeraAssetPreviewActor>& GetPreviewActors() { return PreviewActors; }
+	TRange<double> GetPreviewPlayRange() const;
+	void SetPlayTime(float NewPlayTime, bool bInTickPlayTime);
+	float GetPlayTime() const { return PlayTime; }
+	void SetPreviewProperties(float AnimAssetTime, const FVector& AnimAssetBlendParameters, bool bAnimAssetPlaying);
+
+private:
+	UWorld* GetWorld();
+
+	float PlayTime = 0.f;
+	float DeltaTimeMultiplier = 1.f;
+	/** asset being viewed and edited by this view model. */
+	TWeakObjectPtr<UChimeraAsset> ChimeraAssetPtr;
+	/** Weak pointer to the PreviewScene */
+	TWeakPtr<FChimeraAssetPreviewScene> PreviewScenePtr;
+	/** Actors to be displayed in the preview viewport */
+	TArray<FChimeraAssetPreviewActor> PreviewActors;
+};
+
+/////////////////////////////////////////////////
+// class FChimeraAssetEdMode
+class FChimeraAssetEdMode : public FEdMode
+{
+public:
+	const static FEditorModeID EdModeId;
+
+	virtual void Tick(FEditorViewportClient* ViewportClient, float DeltaTime) override;
+	virtual void Render(const FSceneView* View, FViewport* Viewport, FPrimitiveDrawInterface* PDI) override;
+	virtual bool AllowWidgetMove() override;
+	virtual bool ShouldDrawWidget() const override;
+	virtual bool GetCustomInputCoordinateSystem(FMatrix& InMatrix, void* InData) override;
+
+private:
+	FChimeraAssetViewModel* ViewModel = nullptr;
+};
+
+/////////////////////////////////////////////////
+// class FChimeraAssetViewportClient
+class FChimeraAssetViewportClient : public FEditorViewportClient
+{
+public:
+	FChimeraAssetViewportClient(const TSharedRef<FChimeraAssetPreviewScene>& InPreviewScene, const TSharedRef<SChimeraAssetViewport>& InViewport, const TSharedRef<FChimeraAssetEditor>& InAssetEditor);
+	virtual void Draw(const FSceneView* View, FPrimitiveDrawInterface* PDI) override;
+	virtual void TrackingStarted(const struct FInputEventState& InInputState, bool bIsDragging, bool bNudge) override;
+	virtual void TrackingStopped() override;
+
+	/** Get the preview scene we are viewing */
+	TSharedRef<FChimeraAssetPreviewScene> GetPreviewScene() const { return PreviewScenePtr.Pin().ToSharedRef(); }
+	TSharedRef<FChimeraAssetEditor> GetAssetEditor() const { return AssetEditorPtr.Pin().ToSharedRef();	}
+
+private:
+	/** Preview scene we are viewing */
+	TWeakPtr<FChimeraAssetPreviewScene> PreviewScenePtr;
+	/** Asset editor we are embedded in */
+	TWeakPtr<FChimeraAssetEditor> AssetEditorPtr;
+};
+
+/////////////////////////////////////////////////
+// class FChimeraAssetPreviewScene
+class FChimeraAssetPreviewScene : public FAdvancedPreviewScene
+{
+public:
+	FChimeraAssetPreviewScene(ConstructionValues CVs, const TSharedRef<FChimeraAssetEditor>& Editor);
+	virtual void Tick(float InDeltaTime) override;
+	TSharedRef<FChimeraAssetEditor> GetEditor() const { return EditorPtr.Pin().ToSharedRef(); }
+
+private:
+	/** The asset editor we are embedded in */
+	TWeakPtr<FChimeraAssetEditor> EditorPtr;
+};
+
+/////////////////////////////////////////////////
+// struct FChimeraAssetPreviewRequiredArgs
+struct FChimeraAssetPreviewRequiredArgs
+{
+	FChimeraAssetPreviewRequiredArgs(const TSharedRef<FChimeraAssetEditor>& InAssetEditor, const TSharedRef<FChimeraAssetPreviewScene>& InPreviewScene)
+		: AssetEditor(InAssetEditor)
+		, PreviewScene(InPreviewScene)
+	{
+	}
+
+	TSharedRef<FChimeraAssetEditor> AssetEditor;
+	TSharedRef<FChimeraAssetPreviewScene> PreviewScene;
+};
+
+/////////////////////////////////////////////////
+// class SChimeraAssetViewport
+class SChimeraAssetViewport : public SEditorViewport, public ICommonEditorViewportToolbarInfoProvider
+{
+public:
+	SLATE_BEGIN_ARGS(SChimeraAssetViewport) {}
+	SLATE_END_ARGS();
+
+	void Construct(const FArguments& InArgs, const FChimeraAssetPreviewRequiredArgs& InRequiredArgs);
+	virtual TSharedRef<class SEditorViewport> GetViewportWidget() override;
+	virtual TSharedPtr<FExtender> GetExtenders() const override;
+	virtual void OnFloatingButtonClicked() override;
+
+protected:
+	virtual void BindCommands() override;
+	virtual TSharedRef<FEditorViewportClient> MakeEditorViewportClient() override;
+	virtual TSharedPtr<SWidget> MakeViewportToolbar() override;
+
+	/** The viewport toolbar */
+	TSharedPtr<SCommonEditorViewportToolbarBase> ViewportToolbar;
+	/** Viewport client */
+	TSharedPtr<FChimeraAssetViewportClient> ViewportClient;
+	/** The preview scene that we are viewing */
+	TWeakPtr<FChimeraAssetPreviewScene> PreviewScenePtr;
+	/** Asset editor we are embedded in */
+	TWeakPtr<FChimeraAssetEditor> AssetEditorPtr;
+};
+
+/////////////////////////////////////////////////
+// class SChimeraAssetPreview
+class SChimeraAssetPreview : public SCompoundWidget
+{
+public:
+	DECLARE_DELEGATE_TwoParams(FOnScrubPositionChanged, double, bool)
+	DECLARE_DELEGATE(FOnButtonClickedEvent)
+
+	SLATE_BEGIN_ARGS(SChimeraAssetPreview) {}
+		SLATE_ATTRIBUTE(FLinearColor, SliderColor);
+		SLATE_ATTRIBUTE(double, SliderScrubTime);
+		SLATE_ATTRIBUTE(TRange<double>, SliderViewRange);
+		SLATE_EVENT(FOnScrubPositionChanged, OnSliderScrubPositionChanged);
+		SLATE_EVENT(FOnButtonClickedEvent, OnBackwardEnd);
+		SLATE_EVENT(FOnButtonClickedEvent, OnBackwardStep);
+		SLATE_EVENT(FOnButtonClickedEvent, OnBackward);
+		SLATE_EVENT(FOnButtonClickedEvent, OnPause);
+		SLATE_EVENT(FOnButtonClickedEvent, OnForward);
+		SLATE_EVENT(FOnButtonClickedEvent, OnForwardStep);
+		SLATE_EVENT(FOnButtonClickedEvent, OnForwardEnd);
+	SLATE_END_ARGS();
+
+	void Construct(const FArguments& InArgs, const FChimeraAssetPreviewRequiredArgs& InRequiredArgs);
+
+protected:
+	TAttribute<FLinearColor> SliderColor;
+	TAttribute<double> SliderScrubTime;
+	TAttribute<TRange<double>> SliderViewRange = TRange<double>(0.0, 1.0);
+	FOnScrubPositionChanged OnSliderScrubPositionChanged;
+	FOnButtonClickedEvent OnBackwardEnd;
+	FOnButtonClickedEvent OnBackwardStep;
+	FOnButtonClickedEvent OnBackward;
+	FOnButtonClickedEvent OnPause;
+	FOnButtonClickedEvent OnForward;
+	FOnButtonClickedEvent OnForwardStep;
+	FOnButtonClickedEvent OnForwardEnd;
+};
+
+/////////////////////////////////////////////////
+// class FChimeraAssetEditor
+class FChimeraAssetEditor : public IMultiAnimAssetEditor, public FNotifyHook
+{
+public:
+
+	void InitAssetEditor(const EToolkitMode::Type Mode, const TSharedPtr<IToolkitHost>& InitToolkitHost, UChimeraAsset* ChimeraAsset);
+
+	virtual void RegisterTabSpawners(const TSharedRef<FTabManager>& TabManager) override;
+	virtual void UnregisterTabSpawners(const TSharedRef<FTabManager>& TabManager) override;
+	virtual FName GetToolkitFName() const override;
+	virtual FText GetBaseToolkitName() const override;
+	virtual FText GetToolkitName() const override;
+	virtual FLinearColor GetWorldCentricTabColorScale() const override;
+	virtual FString GetWorldCentricTabPrefix() const override;
+	virtual void SetPreviewProperties(float AnimAssetTime, const FVector& AnimAssetBlendParameters, bool bAnimAssetPlaying) override;
+
+	const UChimeraAsset* GetChimeraAsset() const { return ViewModel.IsValid() ? ViewModel->GetChimeraAsset() : nullptr; }
+	FChimeraAssetViewModel* GetViewModel() const { return ViewModel.Get(); }
+	void PreviewBackwardEnd() { ViewModel->PreviewBackwardEnd(); }
+	void PreviewBackwardStep() { ViewModel->PreviewBackwardStep();	}
+	void PreviewBackward() { ViewModel->PreviewBackward(); }
+	void PreviewPause() { ViewModel->PreviewPause(); }
+	void PreviewForward() { ViewModel->PreviewForward(); }
+	void PreviewForwardStep() { ViewModel->PreviewForwardStep(); }
+	void PreviewForwardEnd() { ViewModel->PreviewForwardEnd(); }
+
+private:
+
+	TSharedRef<SDockTab> SpawnTab_Viewport(const FSpawnTabArgs& Args);
+	TSharedRef<SDockTab> SpawnTab_AssetDetails(const FSpawnTabArgs& Args);
+
+	TSharedPtr<SChimeraAssetPreview> PreviewWidget;
+	TSharedPtr<IDetailsView> EditingAssetWidget;
+	TSharedPtr<FChimeraAssetPreviewScene> PreviewScene;
+	TSharedPtr<FChimeraAssetViewModel> ViewModel;
+};
+
+} // UE::Chimera
+
diff --git a/Engine/Plugins/Experimental/Animation/Chimera/Source/Editor/Public/ChimeraAssetFactory.h b/Engine/Plugins/Experimental/Animation/Chimera/Source/Editor/Public/ChimeraAssetFactory.h
new file mode 100644
--- /dev/null
+++ b/Engine/Plugins/Experimental/Animation/Chimera/Source/Editor/Public/ChimeraAssetFactory.h
@@ -0,0 +1,19 @@
+// Copyright Epic Games, Inc. All Rights Reserved.
+
+#pragma once
+
+#include "CoreMinimal.h"
+#include "UObject/ObjectMacros.h"
+#include "Factories/Factory.h"
+#include "ChimeraAssetFactory.generated.h"
+
+UCLASS()
+class CHIMERAEDITOR_API UChimeraAssetFactory : public UFactory
+{
+	GENERATED_UCLASS_BODY()
+
+	// UFactory interface
+	virtual UObject* FactoryCreateNew(UClass* Class, UObject* InParent, FName Name, EObjectFlags Flags, UObject* Context, FFeedbackContext* Warn) override;
+	virtual FString GetDefaultNewAssetName() const override;
+	// End of UFactory interface
+};
diff --git a/Engine/Plugins/Experimental/Animation/Chimera/Source/Editor/Public/ChimeraEditorDefines.h b/Engine/Plugins/Experimental/Animation/Chimera/Source/Editor/Public/ChimeraEditorDefines.h
new file mode 100644
--- /dev/null
+++ b/Engine/Plugins/Experimental/Animation/Chimera/Source/Editor/Public/ChimeraEditorDefines.h
@@ -0,0 +1,7 @@
+// Copyright Epic Games, Inc. All Rights Reserved.
+
+#pragma once
+
+#include "Logging/LogMacros.h"
+
+DECLARE_LOG_CATEGORY_EXTERN(LogChimeraEditor, Log, All);
diff --git a/Engine/Plugins/Experimental/Animation/Chimera/Source/Runtime/Chimera.Build.cs b/Engine/Plugins/Experimental/Animation/Chimera/Source/Runtime/Chimera.Build.cs
new file mode 100644
--- /dev/null
+++ b/Engine/Plugins/Experimental/Animation/Chimera/Source/Runtime/Chimera.Build.cs
@@ -0,0 +1,26 @@
+// Copyright Epic Games, Inc. All Rights Reserved.
+
+using UnrealBuildTool;
+
+public class Chimera : ModuleRules
+{
+	public Chimera(ReadOnlyTargetRules Target) : base(Target)
+	{
+		PublicDependencyModuleNames.AddRange(
+			new string[]
+			{
+				"AnimationCore",
+                "AnimGraphRuntime",
+                "Core",
+                "CoreUObject",
+                "Engine",
+				"PoseSearch",
+            });
+
+			PrivateDependencyModuleNames.AddRange(
+			new string[]
+			{
+				"BlendStack"
+			});
+	}
+}
diff --git a/Engine/Plugins/Experimental/Animation/Chimera/Source/Runtime/Private/AnimNode_Chimera.cpp b/Engine/Plugins/Experimental/Animation/Chimera/Source/Runtime/Private/AnimNode_Chimera.cpp
new file mode 100644
--- /dev/null
+++ b/Engine/Plugins/Experimental/Animation/Chimera/Source/Runtime/Private/AnimNode_Chimera.cpp
@@ -0,0 +1,197 @@
+// Copyright Epic Games, Inc. All Rights Reserved.
+
+#include "Chimera/AnimNode_Chimera.h"
+#include "Animation/AnimInertializationSyncScope.h"
+#include "Animation/AnimInstanceProxy.h"
+#include "Chimera/ChimeraAsset.h"
+#include "Chimera/ChimeraDefines.h"
+#include "Chimera/ChimeraLibrary.h"
+#include "PoseSearch/PoseHistoryProvider.h"
+#include "PoseSearch/PoseSearchDatabase.h"
+#include "PoseSearch/PoseSearchSchema.h"
+#include "HAL/IConsoleManager.h"
+
+#if ENABLE_ANIM_DEBUG
+TAutoConsoleVariable<bool> CVarAnimNodeChimeraDebug(TEXT("a.AnimNode.Chimera.Debug"), false, TEXT("Turn on visualization debugging for AnimNode Chimera"));
+#endif // ENABLE_ANIM_DEBUG
+
+void FAnimNode_Chimera::GatherDebugData(FNodeDebugData& DebugData)
+{
+	Super::GatherDebugData(DebugData);
+
+	FString DebugLine = DebugData.GetNodeName(this);
+#if ENABLE_ANIM_DEBUG
+	static bool bActive = true;
+	DebugLine += FString::Printf(TEXT("\n - Active: (%s)"), bActive);
+#endif
+	DebugData.AddDebugItem(DebugLine);
+}
+
+void FAnimNode_Chimera::Initialize_AnyThread(const FAnimationInitializeContext& Context)
+{
+	Super::Initialize_AnyThread(Context);
+	Source.Initialize(Context);
+}
+
+void FAnimNode_Chimera::CacheBones_AnyThread(const FAnimationCacheBonesContext& Context)
+{
+	Super::CacheBones_AnyThread(Context);
+	Source.CacheBones(Context);
+}
+
+void FAnimNode_Chimera::UpdateAssetPlayer(const FAnimationUpdateContext& Context)
+{
+	using namespace UE::PoseSearch;
+
+	const bool bNeedsReset =
+		bResetOnBecomingRelevant &&
+		UpdateCounter.HasEverBeenUpdated() &&
+		!UpdateCounter.WasSynchronizedCounter(Context.AnimInstanceProxy->GetUpdateCounter());
+
+	UpdateCounter.SynchronizeWith(Context.AnimInstanceProxy->GetUpdateCounter());
+
+	GetEvaluateGraphExposedInputs().Execute(Context);
+	
+	if (bNeedsReset)
+	{
+		Reset();
+		BlendLerp = 0.f;
+		TranslationWarpLerp = 0.f;
+		RotationWarpLerp = 0.f;
+	}
+
+	bool bInteracting = false;
+	bool bExecuteBlendTo = false;
+	const float DeltaTime = Context.GetDeltaTime();
+	FPoseHistoryProvider* PoseHistoryProvider = Context.GetMessage<FPoseHistoryProvider>();
+	if (!PoseHistoryProvider)
+	{
+		UE_LOG(LogChimera, Error, TEXT("FAnimNode_Chimera::Update_AnyThread couldn't find the FPoseHistoryProvider"));
+	}
+	else
+	{
+		check(Context.AnimInstanceProxy);
+		const FChimeraBlueprintResult Result = UChimeraLibrary::ChimeraQuery(Availabilities, Context.AnimInstanceProxy->GetAnimInstanceObject(), PoseHistoryProvider->GetHistoryCollector(), bValidateResultAgainstAvailabilities);
+
+		if (const UMultiAnimAsset* MultiAnimAsset = Cast<UMultiAnimAsset>(Result.SelectedAnimation))
+		{
+			check(Result.SelectedDatabase != nullptr);
+
+			UAnimationAsset* AnimationAsset = MultiAnimAsset->GetAnimationAsset(Result.Role);
+			check(AnimationAsset);
+
+			FullAlignedActorRootBoneTransform = Result.FullAlignedActorRootBoneTransform;
+			bInteracting = true;
+
+			if (AnimPlayers.IsEmpty())
+			{
+				bExecuteBlendTo = true;
+			}
+			else
+			{
+				const FBlendStackAnimPlayer& MainAnimPlayer = AnimPlayers[0];
+				const UAnimationAsset* PlayingAnimationAsset = MainAnimPlayer.GetAnimationAsset();
+
+				if (AnimationAsset != PlayingAnimationAsset ||
+					Result.bIsMirrored != MainAnimPlayer.GetMirror() ||
+					Result.BlendParameters != MainAnimPlayer.GetBlendParameters() ||
+					!Result.bIsContinuingPoseSearch)
+				{
+					bExecuteBlendTo = true;
+				}
+			}
+
+			if (bExecuteBlendTo)
+			{
+				const FPoseSearchRoledSkeleton* RoledSkeleton = Result.SelectedDatabase->Schema->GetRoledSkeleton(Result.Role);
+				check(RoledSkeleton);
+
+				BlendTo(Context, AnimationAsset, Result.SelectedTime, Result.bLoop, Result.bIsMirrored, RoledSkeleton->MirrorDataTable.Get(),
+					BlendTime, BlendProfile, BlendOption, bUseInertialBlend, Result.BlendParameters, Result.WantedPlayRate);
+			}
+
+			// @todo: do we still need FAnimInertializationSyncScope?
+			//const bool bDidBlendToRequestAnInertialBlend = bExecuteBlendTo && bUseInertialBlend;
+			//UE::Anim::TOptionalScopedGraphMessage<UE::Anim::FAnimInertializationSyncScope> InertializationSync(bDidBlendToRequestAnInertialBlend, Context);
+
+			UpdatePlayRate(Result.WantedPlayRate);
+		}
+	}
+	
+	// calculating the translation and rotation warp lerps, used to wapr the root transform towards the last computed FullAlignedActorRootBoneTransform
+	const float Sign = bInteracting ? 1.f : -1.f;
+	if (BlendTime > UE_KINDA_SMALL_NUMBER)
+	{
+		BlendLerp = FMath::Clamp(BlendLerp + (Sign * DeltaTime / BlendTime), 0.f, 1.f);
+	}
+	if (InitialTranslationWarpTime > UE_KINDA_SMALL_NUMBER)
+	{
+		TranslationWarpLerp = FMath::Clamp(TranslationWarpLerp + (Sign * DeltaTime / InitialTranslationWarpTime), 0.f, 1.f);
+	}
+	if (InitialRotationWarpTime > UE_KINDA_SMALL_NUMBER)
+	{
+		RotationWarpLerp = FMath::Clamp(RotationWarpLerp + (Sign * DeltaTime / InitialRotationWarpTime), 0.f, 1.f);
+	}
+
+	// updating source
+	const float SourceBlendLerp = 1.f - BlendLerp;
+	FAnimationUpdateContext SourceContext = Context.FractionalWeightAndRootMotion(SourceBlendLerp, SourceBlendLerp);
+	if (BlendLerp > UE_KINDA_SMALL_NUMBER)
+	{
+		SourceContext = SourceContext.AsInactive();
+	}
+	Source.Update(SourceContext);
+	
+	// updating blend stack
+	if (BlendLerp < UE_KINDA_SMALL_NUMBER && DeltaTime > UE_KINDA_SMALL_NUMBER)
+	{
+		// resetting the blendstack if there's no BlendLerp weight to it
+		Reset();
+	}
+	FAnimationUpdateContext BlendStackContext = Context.FractionalWeightAndRootMotion(BlendLerp, BlendLerp);
+	Super::UpdateAssetPlayer(BlendStackContext);
+
+#if ENABLE_ANIM_DEBUG
+	if (UE_TRACE_CHANNELEXPR_IS_ENABLED(AnimationChannel))
+	{
+		TRACE_ANIM_NODE_VALUE(Context, *FString("Interacting"), bInteracting);
+		TRACE_ANIM_NODE_VALUE(Context, *FString("ExecuteBlendTo"), bExecuteBlendTo);
+		TRACE_ANIM_NODE_VALUE(Context, *FString("BlendLerp"), BlendLerp);
+		TRACE_ANIM_NODE_VALUE(Context, *FString("TranslationWarpLerp"), TranslationWarpLerp);
+		TRACE_ANIM_NODE_VALUE(Context, *FString("RotationWarpLerp"), RotationWarpLerp);
+	}
+#endif // ENABLE_ANIM_DEBUG
+}
+
+void FAnimNode_Chimera::Evaluate_AnyThread(FPoseContext& Output)
+{
+	check(Output.AnimInstanceProxy);
+
+	// @todo: do we still need to evaluate Source if blendstack weight is 100%?
+	// evaluating Source to get the base pose 
+	Source.Evaluate(Output);
+
+	if (!AnimPlayers.IsEmpty())
+	{
+		Super::Evaluate_AnyThread(Output);
+	}
+
+	if (TranslationWarpLerp > UE_KINDA_SMALL_NUMBER || RotationWarpLerp > UE_KINDA_SMALL_NUMBER)
+	{
+		const FTransform ComponentTransform = Output.AnimInstanceProxy->GetComponentTransform();
+		const FTransform FullAlignedActorRootBoneLocalTransform = FullAlignedActorRootBoneTransform.GetRelativeTransform(ComponentTransform);
+
+#if ENABLE_ANIM_DEBUG
+		if (CVarAnimNodeChimeraDebug.GetValueOnAnyThread())
+		{
+			Output.AnimInstanceProxy->AnimDrawDebugCoordinateSystem(FullAlignedActorRootBoneTransform.GetLocation(), FullAlignedActorRootBoneTransform.Rotator(), 10.0f, false, 0.f, 0.f, SDPG_Foreground);
+			Output.AnimInstanceProxy->AnimDrawDebugCoordinateSystem(ComponentTransform.GetLocation(), ComponentTransform.Rotator(), 10.0f, false, 0.f, 0.f, SDPG_Foreground);
+		}
+#endif // ENABLE_ANIM_DEBUG
+
+		const FCompactPoseBoneIndex RootBoneIndex(0);
+		FTransform& RootBoneTransform = Output.Pose[RootBoneIndex];
+		RootBoneTransform.SetTranslation(FMath::Lerp(RootBoneTransform.GetTranslation(), FullAlignedActorRootBoneLocalTransform.GetTranslation(), TranslationWarpLerp));
+		RootBoneTransform.SetRotation(FQuat::Slerp(RootBoneTransform.GetRotation(), FullAlignedActorRootBoneLocalTransform.GetRotation(), RotationWarpLerp));
+	}
+}
diff --git a/Engine/Plugins/Experimental/Animation/Chimera/Source/Runtime/Private/ChimeraAsset.cpp b/Engine/Plugins/Experimental/Animation/Chimera/Source/Runtime/Private/ChimeraAsset.cpp
new file mode 100644
--- /dev/null
+++ b/Engine/Plugins/Experimental/Animation/Chimera/Source/Runtime/Private/ChimeraAsset.cpp
@@ -0,0 +1,356 @@
+// Copyright Epic Games, Inc. All Rights Reserved.
+
+#include "Chimera/ChimeraAsset.h"
+#include "Animation/AnimSequence.h"
+#include "Animation/BlendSpace.h"
+#include "DrawDebugHelpers.h"
+#include "PoseSearch/PoseSearchAssetSampler.h"
+
+bool UChimeraAsset::IsLooping() const
+{
+	float CommonPlayLength = -1.f;
+	for (const FChimeraAssetItem& Item : Items)
+	{
+		if (const UAnimationAsset* AnimationAsset = Item.Animation.Get())
+		{
+			if (const UAnimSequenceBase* SequenceBase = Cast<UAnimSequenceBase>(AnimationAsset))
+			{
+				if (!SequenceBase->bLoop)
+				{
+					return false;
+				}
+			}
+			else if (const UBlendSpace* BlendSpace = Cast<UBlendSpace>(AnimationAsset))
+			{
+				if (!BlendSpace->bLoop)
+				{
+					return false;
+				}
+			}
+			else
+			{
+				unimplemented();
+			}
+
+			if (CommonPlayLength < 0.f)
+			{
+				CommonPlayLength = AnimationAsset->GetPlayLength();
+			}
+			else if (!FMath::IsNearlyEqual(CommonPlayLength, AnimationAsset->GetPlayLength()))
+			{
+				return false;
+			}
+		}
+	}
+	return true;
+}
+
+bool UChimeraAsset::HasRootMotion() const
+{
+	bool bHasAtLeastOneValidItem = false;
+	bool bHasRootMotion = true;
+
+	for (const FChimeraAssetItem& Item : Items)
+	{
+		if (const UAnimationAsset* AnimationAsset = Item.Animation.Get())
+		{
+			if (const UAnimSequenceBase* SequenceBase = Cast<UAnimSequenceBase>(AnimationAsset))
+			{
+				bHasRootMotion &= SequenceBase->HasRootMotion();
+			}
+			else if (const UBlendSpace* BlendSpace = Cast<UBlendSpace>(AnimationAsset))
+			{
+				BlendSpace->ForEachImmutableSample([&bHasRootMotion](const FBlendSample& Sample)
+				{
+					if (const UAnimSequence* Sequence = Sample.Animation.Get())
+					{
+						bHasRootMotion &= Sequence->HasRootMotion();
+					}
+				});
+			}
+			else
+			{
+				unimplemented();
+			}
+			bHasAtLeastOneValidItem = true;
+		}
+	}
+
+	return bHasAtLeastOneValidItem && bHasRootMotion;
+}
+
+float UChimeraAsset::GetPlayLength() const
+{
+	float MaxPlayLength = 0.f;
+	for (const FChimeraAssetItem& Item : Items)
+	{
+		if (const UAnimationAsset* AnimationAsset = Item.Animation.Get())
+		{
+			MaxPlayLength = FMath::Max(MaxPlayLength, AnimationAsset->GetPlayLength());
+		}
+	}
+	return MaxPlayLength;
+}
+
+#if WITH_EDITOR
+int32 UChimeraAsset::GetFrameAtTime(float Time) const
+{
+	const UAnimationAsset* MaxPlayLengthAnim = nullptr;
+	float MaxPlayLength = -1.f;
+	for (const FChimeraAssetItem& Item : Items)
+	{
+		if (const UAnimationAsset* AnimationAsset = Item.Animation.Get())
+		{
+			const float PlayLength = AnimationAsset->GetPlayLength();
+			if (PlayLength > MaxPlayLength)
+			{
+				MaxPlayLength = PlayLength;
+				MaxPlayLengthAnim = AnimationAsset;
+			}
+		}
+	}
+	
+	if (MaxPlayLengthAnim)
+	{
+		if (const UAnimSequenceBase* SequenceBase = Cast<UAnimSequenceBase>(MaxPlayLengthAnim))
+		{
+			return SequenceBase->GetFrameAtTime(Time);
+		}
+
+		if (const UBlendSpace* BlendSpace = Cast<UBlendSpace>(MaxPlayLengthAnim))
+		{
+			// returning the percentage of time as value to diplay in the pose search debugger (NoTe: BlendSpace->GetPlayLength() is one)
+			return FMath::RoundToInt(Time * 100.f);
+		}
+
+		unimplemented();
+	}
+
+	return 0;
+}
+#endif // WITH_EDITOR
+
+FQuat UChimeraAsset::FindReferenceOrientation(const TArrayView<const FTransform> Transforms, const TArrayView<int32> SortedByWarpingWeightRotationItemIndex) const
+{
+	const int32 ItemsNum = Items.Num();
+
+	check(ItemsNum > 0);
+	check(Transforms.Num() == ItemsNum);
+	check(Transforms.Num() == SortedByWarpingWeightRotationItemIndex.Num());
+
+	const int32 LastItemIndex = ItemsNum - 1;
+	if (ItemsNum > 1)
+	{
+		FVector OtherItemsPositionsSum = Transforms[SortedByWarpingWeightRotationItemIndex[0]].GetTranslation();
+		for (int32 ItemIndex = 1; ItemIndex < LastItemIndex; ++ItemIndex)
+		{
+			OtherItemsPositionsSum += Transforms[SortedByWarpingWeightRotationItemIndex[ItemIndex]].GetTranslation();
+		}
+		
+		const FVector OtherItemsPositionAverage = OtherItemsPositionsSum / LastItemIndex;
+		const FVector DeltaPosition = OtherItemsPositionAverage - Transforms[SortedByWarpingWeightRotationItemIndex[LastItemIndex]].GetTranslation();
+
+		if (!DeltaPosition.IsNearlyZero())
+		{
+			return DeltaPosition.ToOrientationQuat();
+		}
+	}
+
+	return Transforms[SortedByWarpingWeightRotationItemIndex[LastItemIndex]].GetRotation();
+}
+
+FVector UChimeraAsset::FindReferencePosition(const TArrayView<const FTransform> Transforms, const TArrayView<float> NormalizedWarpingWeightTranslation) const
+{
+	const int32 ItemsNum = Items.Num();
+	
+	check(ItemsNum > 0);
+	check(Transforms.Num() == ItemsNum);
+	check(Transforms.Num() == NormalizedWarpingWeightTranslation.Num());
+
+	FVector PositionsSum = FVector::ZeroVector;
+	for (int32 ItemIndex = 0; ItemIndex < ItemsNum; ++ItemIndex)
+	{
+		PositionsSum += Transforms[ItemIndex].GetTranslation() * NormalizedWarpingWeightTranslation[ItemIndex];
+	}
+
+	return PositionsSum;
+}
+
+UAnimationAsset* UChimeraAsset::GetAnimationAsset(const UE::PoseSearch::FRole& Role) const
+{
+	for (const FChimeraAssetItem& Item : Items)
+	{
+		if (Item.Role == Role)
+		{
+			return Item.Animation;
+		}
+	}
+	return nullptr;
+}
+
+FTransform UChimeraAsset::GetOrigin(const UE::PoseSearch::FRole& Role) const
+{
+	for (const FChimeraAssetItem& Item : Items)
+	{
+		if (Item.Role == Role)
+		{
+			return Item.Origin;
+		}
+	}
+	return FTransform::Identity;
+}
+
+#if WITH_EDITOR
+FTransform UChimeraAsset::GetDebugWarpOrigin(const UE::PoseSearch::FRole& Role, bool bComposeWithDebugWarpOffset) const
+{
+	for (int32 ItemIndex = 0; ItemIndex < Items.Num(); ++ItemIndex)
+	{
+		const FChimeraAssetItem& Item = Items[ItemIndex];
+		if (Item.Role == Role)
+		{
+#if WITH_EDITORONLY_DATA
+			if (bComposeWithDebugWarpOffset && bEnableDebugWarp && DebugWarpOffsets.IsValidIndex(ItemIndex))
+			{
+				return DebugWarpOffsets[ItemIndex] * Item.Origin;
+			}
+#endif // WITH_EDITORONLY_DATA
+
+			return Item.Origin;
+		}
+	}
+	return FTransform::Identity;
+}
+#endif // WITH_EDITOR
+
+void UChimeraAsset::CalculateWarpTransforms(float Time, const TArrayView<const FTransform> ActorRootBoneTransforms, TArrayView<FTransform> FullAlignedActorRootBoneTransforms, const UWorld* DebugDrawWorld) const
+{
+	check(ActorRootBoneTransforms.Num() == GetNumRoles());
+	check(FullAlignedActorRootBoneTransforms.Num() == GetNumRoles());
+
+	const int32 ItemsNum = Items.Num();
+	if (ItemsNum == 0)
+	{
+		return;
+	}
+
+	TArray<FTransform, TInlineAllocator<UE::PoseSearch::PreallocatedRolesNum>> AssetRootBoneTransforms;
+	AssetRootBoneTransforms.SetNum(ItemsNum);
+
+	// ItemIndex is the RoleIndex and Role = Item.Role
+	for (int32 ItemIndex = 0; ItemIndex < ItemsNum; ++ItemIndex)
+	{
+		const FChimeraAssetItem& Item = Items[ItemIndex];
+
+		// sampling the AnimationAsset to extract the current time transform and the initial (time of 0) transform
+		const UE::PoseSearch::FAnimationAssetSampler Sampler(Item.Animation, Item.Origin);
+		AssetRootBoneTransforms[ItemIndex] = Sampler.ExtractRootTransform(Time);
+
+#if DO_CHECK
+		// array containing the bone index of the root bone (0)
+		TArray<uint16, TInlineAllocator<1>> BoneIndices;
+		BoneIndices.SetNumZeroed(1);
+
+		// extracting the pose, containing only the root bone from the Sampler 
+		FMemMark Mark(FMemStack::Get());
+		FCompactPose Pose;
+		FBoneContainer BoneContainer;
+		BoneContainer.InitializeTo(BoneIndices, UE::Anim::FCurveFilterSettings(UE::Anim::ECurveFilterMode::DisallowAll), *Items[ItemIndex].Animation->GetSkeleton());
+		Pose.SetBoneContainer(&BoneContainer);
+		Sampler.ExtractPose(Time, Pose);
+
+		// making sure the animation root bone transform is Identity, so we can confuse the root with the root BONE transform and preserve performances!
+		const FTransform& RootBoneTransform = Pose.GetBones()[0];
+		check(RootBoneTransform.Equals(FTransform::Identity));
+#endif // DO_CHECK
+	}
+
+	TArray<int32, TInlineAllocator<UE::PoseSearch::PreallocatedRolesNum>> SortedByWarpingWeightRotationItemIndex;
+	TArray<float, TInlineAllocator<UE::PoseSearch::PreallocatedRolesNum>> NormalizedWarpingWeightTranslation;
+	SortedByWarpingWeightRotationItemIndex.SetNum(ItemsNum);
+	NormalizedWarpingWeightTranslation.SetNum(ItemsNum);
+
+	float WarpingWeightTranslationSum = 0.f;
+	float WarpingWeightRotationSum = 0.f;
+	for (int32 ItemIndex = 0; ItemIndex < ItemsNum; ++ItemIndex)
+	{
+		SortedByWarpingWeightRotationItemIndex[ItemIndex] = ItemIndex;
+		WarpingWeightTranslationSum += Items[ItemIndex].WarpingWeightTranslation;
+		WarpingWeightRotationSum += Items[ItemIndex].WarpingWeightRotation;
+	}
+
+	const float NormalizedHomogeneousWeight = 1.f / ItemsNum;
+	if (WarpingWeightTranslationSum > UE_KINDA_SMALL_NUMBER)
+	{
+		for (int32 ItemIndex = 0; ItemIndex < ItemsNum; ++ItemIndex)
+		{
+			NormalizedWarpingWeightTranslation[ItemIndex] = Items[ItemIndex].WarpingWeightTranslation / WarpingWeightTranslationSum;
+		}
+	}
+	else
+	{
+		for (int32 ItemIndex = 0; ItemIndex < ItemsNum; ++ItemIndex)
+		{
+			NormalizedWarpingWeightTranslation[ItemIndex] = NormalizedHomogeneousWeight;
+		}
+	}
+
+	if (WarpingWeightRotationSum > UE_KINDA_SMALL_NUMBER)
+	{
+		SortedByWarpingWeightRotationItemIndex.Sort([this](const int32 A, const int32 B)
+			{
+				return Items[A].WarpingWeightRotation < Items[B].WarpingWeightRotation;
+			});
+	}
+
+	const FQuat AssetReferenceOrientation = FindReferenceOrientation(AssetRootBoneTransforms, SortedByWarpingWeightRotationItemIndex);
+	const FQuat ActorsReferenceOrientation = FindReferenceOrientation(ActorRootBoneTransforms, SortedByWarpingWeightRotationItemIndex);
+	
+	FQuat WeightedActorsReferenceOrientation = ActorsReferenceOrientation;
+	if (WarpingWeightRotationSum > UE_KINDA_SMALL_NUMBER)
+	{
+		// ItemIndex are in order of WarpingWeightRotation. the last one is the one with the highest most WarpingWeightRotation, the most "important"
+		for (int32 ItemIndex : SortedByWarpingWeightRotationItemIndex)
+		{
+			const FChimeraAssetItem& Item = Items[ItemIndex];
+			const float NormalizedWarpingWeightRotation = Item.WarpingWeightRotation / WarpingWeightRotationSum;
+			if (NormalizedWarpingWeightRotation > NormalizedHomogeneousWeight)
+			{
+				// NormalizedHomogeneousWeight is one only if ItemsNum is one, 
+				// BUT NormalizedWarpingWeightRotation > NormalizedHomogeneousWeight should always be false
+				check(!FMath::IsNearlyEqual(NormalizedHomogeneousWeight, 1.f));
+
+				// how much this item wants to reorient the ReferenceOrientation from the homogeneous "fair" value
+				const float SlerpParam = (NormalizedWarpingWeightRotation - NormalizedHomogeneousWeight) / (1.f - NormalizedHomogeneousWeight);
+
+				// calculating the reference orientation relative to the character
+				// AssetReferenceOrientation in actor world orientation
+				const FQuat ActorAssetReferenceOrientation = ActorRootBoneTransforms[ItemIndex].GetRotation() * (AssetRootBoneTransforms[ItemIndex].GetRotation().Inverse() * AssetReferenceOrientation);
+
+				WeightedActorsReferenceOrientation = FQuat::Slerp(WeightedActorsReferenceOrientation, ActorAssetReferenceOrientation, SlerpParam);
+			}
+		}
+	}
+
+	const FVector AssetReferencePosition = FindReferencePosition(AssetRootBoneTransforms, NormalizedWarpingWeightTranslation);
+	const FVector ActorsReferencePosition = FindReferencePosition(ActorRootBoneTransforms, NormalizedWarpingWeightTranslation);
+
+#if ENABLE_DRAW_DEBUG
+	DrawDebugCoordinateSystem(DebugDrawWorld, AssetReferencePosition, AssetReferenceOrientation.Rotator(), 10.f, false, 0.f, SDPG_Foreground);
+	DrawDebugCoordinateSystem(DebugDrawWorld, ActorsReferencePosition, ActorsReferenceOrientation.Rotator(), 10.f, false, 0.f, SDPG_Foreground);
+	DrawDebugCoordinateSystem(DebugDrawWorld, ActorsReferencePosition, WeightedActorsReferenceOrientation.Rotator(), 20.f, false, 0.f, SDPG_Foreground);
+#endif // ENABLE_DRAW_DEBUG
+
+	// aligning all the actors to ActorsReferencePosition, WeightedActorsReferenceOrientation
+	const FTransform AssetReferenceTransform(AssetReferenceOrientation, AssetReferencePosition);
+	const FTransform ActorsReferenceTransform(WeightedActorsReferenceOrientation, ActorsReferencePosition);
+	const FTransform AssetReferenceInverseTransform = AssetReferenceTransform.Inverse();
+	
+	for (int32 ItemIndex = 0; ItemIndex < ItemsNum; ++ItemIndex)
+	{
+		FullAlignedActorRootBoneTransforms[ItemIndex] = (AssetRootBoneTransforms[ItemIndex] * AssetReferenceInverseTransform) * ActorsReferenceTransform;
+
+#if ENABLE_DRAW_DEBUG
+		DrawDebugCoordinateSystem(DebugDrawWorld, FullAlignedActorRootBoneTransforms[ItemIndex].GetTranslation(), FullAlignedActorRootBoneTransforms[ItemIndex].Rotator(), 10.f, false, 0.f, SDPG_Foreground);
+#endif // ENABLE_DRAW_DEBUG
+	}
+}
\ No newline at end of file
diff --git a/Engine/Plugins/Experimental/Animation/Chimera/Source/Runtime/Private/ChimeraDefines.cpp b/Engine/Plugins/Experimental/Animation/Chimera/Source/Runtime/Private/ChimeraDefines.cpp
new file mode 100644
--- /dev/null
+++ b/Engine/Plugins/Experimental/Animation/Chimera/Source/Runtime/Private/ChimeraDefines.cpp
@@ -0,0 +1,5 @@
+// Copyright Epic Games, Inc. All Rights Reserved.
+
+#include "Chimera/ChimeraDefines.h"
+
+DEFINE_LOG_CATEGORY(LogChimera);
diff --git a/Engine/Plugins/Experimental/Animation/Chimera/Source/Runtime/Private/ChimeraIslandComponent.cpp b/Engine/Plugins/Experimental/Animation/Chimera/Source/Runtime/Private/ChimeraIslandComponent.cpp
new file mode 100644
--- /dev/null
+++ b/Engine/Plugins/Experimental/Animation/Chimera/Source/Runtime/Private/ChimeraIslandComponent.cpp
@@ -0,0 +1,486 @@
+// Copyright Epic Games, Inc. All Rights Reserved.
+
+#include "Chimera/ChimeraIslandComponent.h"
+#include "Animation/AnimClassInterface.h"
+#include "Animation/AnimInstance.h"
+#include "Components/SkeletalMeshComponent.h"
+#include "Chimera/ChimeraAsset.h"
+#include "Chimera/ChimeraDefines.h"
+#include "Chimera/ChimeraLibrary.h"
+#include "DrawDebugHelpers.h"
+#include "GameFramework/CharacterMovementComponent.h"
+#include "PoseSearch/AnimNode_PoseSearchHistoryCollector.h"
+#include "PoseSearch/PoseSearchLibrary.h"
+#include "PoseSearch/PoseSearchDatabase.h"
+#include "PoseSearch/PoseSearchSchema.h"
+
+namespace UE::Chimera
+{
+#if ENABLE_ANIM_DEBUG
+	static TAutoConsoleVariable<bool> CVarChimeraShowIslands(TEXT("a.Chimera.ShowIslands"), false, TEXT("Show Chimera Islands"));
+#endif
+
+	static UE::Chimera::FSearchResult InitSearchResult(const UE::PoseSearch::FSearchResult& SearchResult, int32 SearchIndex, const UE::Chimera::FSearchContext& SearchContext)
+	{
+		UE::Chimera::FSearchResult ChimeraSearchResult;
+		static_cast<UE::PoseSearch::FSearchResult&>(ChimeraSearchResult) = SearchResult;
+		ChimeraSearchResult.SearchIndex = SearchIndex;
+
+		// @todo: WIP! calculating warping: currently supporting only UChimeraAsset(s), but we should generalize for UMultiAnimAsset(s)
+		if (const UE::PoseSearch::FSearchIndexAsset* SearchIndexAsset = SearchResult.GetSearchIndexAsset())
+		{
+			if (const FPoseSearchDatabaseMultiAnimAsset* DatabaseMultiAnimAsset = SearchResult.Database->GetDatabaseAnimationAsset<FPoseSearchDatabaseMultiAnimAsset>(*SearchIndexAsset))
+			{
+				if (const UChimeraAsset* ChimeraAsset = Cast<UChimeraAsset>(DatabaseMultiAnimAsset->GetAnimationAsset()))
+				{
+					check(ChimeraAsset->GetNumRoles() == SearchContext.AnimInstances.Num());
+					check(ChimeraAsset->GetNumRoles() == SearchContext.HistoryCollectors.Num());
+					check(ChimeraAsset->GetNumRoles() == SearchContext.Roles.Num());
+
+					UE::PoseSearch::FRoleToIndex SearchContextRoleToIndex;
+					SearchContextRoleToIndex.Reserve(SearchContext.Roles.Num());
+					for (int32 SearchContextRoleIndex = 0; SearchContextRoleIndex < SearchContext.Roles.Num(); ++SearchContextRoleIndex)
+					{
+						SearchContextRoleToIndex.Add(SearchContext.Roles[SearchContextRoleIndex]) = SearchContextRoleIndex;
+					}
+
+					// mapping ChimeraAsset Roles to SearchContext.* indexes
+					TArray<FTransform, TInlineAllocator<UE::PoseSearch::PreallocatedRolesNum>> ActorRootBoneTransforms;
+					const UWorld* DebugDrawWorld = nullptr;
+					for (int32 ChimeraAssetRoleIndex = 0; ChimeraAssetRoleIndex < ChimeraAsset->GetNumRoles(); ++ChimeraAssetRoleIndex)
+					{
+						const UE::PoseSearch::FRole& ChimeraAssetRole = ChimeraAsset->GetRole(ChimeraAssetRoleIndex);
+						const int32 SearchContextIndex = SearchContextRoleToIndex[ChimeraAssetRole];
+
+						if (UAnimInstance* AnimInstance = SearchContext.AnimInstances[SearchContextIndex].Get())
+						{
+							const FTransform RootBoneTransform = AnimInstance->GetSkelMeshComponent()->GetBoneTransform(0);
+							ActorRootBoneTransforms.Add(RootBoneTransform);
+					
+							if (!DebugDrawWorld)
+							{
+								DebugDrawWorld = AnimInstance->GetWorld();
+							}
+						}
+					}
+
+					if (ActorRootBoneTransforms.Num() == ChimeraAsset->GetNumRoles())
+					{
+						// FullAlignedActorRootBoneTransforms is mapped to the ChimeraAsset roles:
+						// FullAlignedActorRootBoneTransforms[0] is for ChimeraAsset->GetRole(0)
+						TArray<FTransform, TInlineAllocator<UE::PoseSearch::PreallocatedRolesNum>> FullAlignedActorRootBoneTransforms;
+						FullAlignedActorRootBoneTransforms.SetNum(ChimeraAsset->GetNumRoles());
+
+						ChimeraAsset->CalculateWarpTransforms(SearchResult.AssetTime, ActorRootBoneTransforms, FullAlignedActorRootBoneTransforms, DebugDrawWorld);
+						ChimeraSearchResult.FullAlignedActorRootBoneTransforms.SetNum(ChimeraAsset->GetNumRoles());
+
+						for (int32 ChimeraAssetRoleIndex = 0; ChimeraAssetRoleIndex < ChimeraAsset->GetNumRoles(); ++ChimeraAssetRoleIndex)
+						{
+							const UE::PoseSearch::FRole& ChimeraAssetRole = ChimeraAsset->GetRole(ChimeraAssetRoleIndex);
+							const int32 SearchContextIndex = SearchContextRoleToIndex[ChimeraAssetRole];
+
+							ChimeraSearchResult.FullAlignedActorRootBoneTransforms[SearchContextIndex] = FullAlignedActorRootBoneTransforms[ChimeraAssetRoleIndex];
+						}
+					}
+				}
+			}
+		}
+
+		return ChimeraSearchResult;
+	}
+
+	bool FSearchContext::IsValid() const
+	{
+		if (Database == nullptr)
+		{
+			return false;
+		}
+
+		const int32 Num = AnimInstances.Num();
+		if (Num < 2)
+		{
+			return false;
+		}
+
+		if (Num != HistoryCollectors.Num())
+		{
+			return false;
+		}
+
+		if (Num != Roles.Num())
+		{
+			return false;
+		}
+
+		for (int32 IndexA = 0; IndexA < Num; ++IndexA)
+		{
+			if (AnimInstances[IndexA] == nullptr)
+			{
+				return false;
+			}
+
+			for (int32 IndexB = IndexA + 1; IndexB < Num; ++IndexB)
+			{
+				if (AnimInstances[IndexA] == AnimInstances[IndexB])
+				{
+					return false;
+				}
+			}
+		}
+
+		for (int32 IndexA = 0; IndexA < Num; ++IndexA)
+		{
+			for (int32 IndexB = IndexA + 1; IndexB < Num; ++IndexB)
+			{
+				if (Roles[IndexA] == Roles[IndexB])
+				{
+					return false;
+				}
+			}
+		}
+
+		for (int32 IndexA = 0; IndexA < Num; ++IndexA)
+		{
+			if (HistoryCollectors[IndexA] == nullptr)
+			{
+				return false;
+			}
+
+			for (int32 IndexB = IndexA + 1; IndexB < Num; ++IndexB)
+			{
+				if (HistoryCollectors[IndexA] == HistoryCollectors[IndexB])
+				{
+					return false;
+				}
+			}
+		}
+
+		return true;
+	}
+
+	bool FSearchContext::IsEquivalent(const FSearchContext& Other) const
+	{
+		if (Database != Other.Database)
+		{
+			return false;
+		}
+
+		const int32 Num = AnimInstances.Num();
+		if (Num != Other.AnimInstances.Num())
+		{
+			return false;
+		}
+
+		int32 CommonRoledAnimInstances = 0;
+		for (int32 IndexThis = 0; IndexThis < Num; ++IndexThis)
+		{
+			for (int32 IndexOther = 0; IndexOther < Num; ++IndexOther)
+			{
+				if (AnimInstances[IndexThis] == Other.AnimInstances[IndexOther] &&
+					HistoryCollectors[IndexThis] == Other.HistoryCollectors[IndexOther] &&
+					Roles[IndexThis] == Other.Roles[IndexOther])
+				{
+					++CommonRoledAnimInstances;
+				}
+			}
+		}
+
+		// using >= in case there are duplicated animinstances in this or Other (this->IsValid() should be false!)
+		if (CommonRoledAnimInstances >= Num)
+		{
+			return true;
+		}
+
+		return false;
+	}
+}
+
+void UChimeraIslandComponent::TickComponent(float DeltaTime, enum ELevelTick TickType, FActorComponentTickFunction* ThisTickFunction)
+{
+	// Called before any skeletal mesh component tick, when there aren't animation jobs flying. No need to FScopeLock Lock(&Mutex);
+	using namespace UE::Chimera;
+
+	// generating trajectories before running any of the skeletal mesh component ticks
+	for (FSearchContext& SearchContext : SearchContexts)
+	{
+		for (int32 Index = 0; Index < SearchContext.AnimInstances.Num(); ++Index)
+		{
+			if (UAnimInstance* AnimInstance = SearchContext.AnimInstances[Index].Get())
+			{
+				// since UChimeraIslandComponent has a tick dependency with the USkeletalMeshComponent it's safe modify the FAnimNode_PoseSearchHistoryCollector_Base
+				FAnimNode_PoseSearchHistoryCollector_Base* HistoryCollector = const_cast<FAnimNode_PoseSearchHistoryCollector_Base*>(SearchContext.HistoryCollectors[Index]);
+				check(HistoryCollector);
+				HistoryCollector->GenerateTrajectory(AnimInstance);
+			}
+		}
+	}
+}
+
+void UChimeraIslandComponent::DebugDraw(const FColor& Color) const
+{
+#if ENABLE_DRAW_DEBUG
+	using namespace UE::Chimera;
+
+	if (CVarChimeraShowIslands.GetValueOnAnyThread())
+	{
+		for (const FSearchContext& SearchContext : SearchContexts)
+		{
+			for (int32 Index = 0; Index < SearchContext.AnimInstances.Num(); ++Index)
+			{
+				if (const UAnimInstance* AnimInstance = SearchContext.AnimInstances[Index].Get())
+				{
+					const FVector Position = AnimInstance->GetSkelMeshComponent()->GetComponentLocation();
+					const float BroadPhaseRadius = SearchContext.BroadPhaseRadiuses[Index];
+					DrawDebugCircle(GetWorld(), Position, BroadPhaseRadius, 40, Color, false, 0.f, SDPG_Foreground, 0.f, FVector::XAxisVector, FVector::YAxisVector, false);
+				}
+			}
+		}
+	}
+#endif // ENABLE_DRAW_DEBUG
+}
+
+void UChimeraIslandComponent::InjectToActor(AActor* Actor)
+{
+	// Called by UChimeraSubsystem::Tick when there aren't animation jobs flying. No need to FScopeLock Lock(&Mutex);
+	if (Actor)
+	{
+		if (UCharacterMovementComponent* CharacterMovementComponent = Actor->GetComponentByClass<UCharacterMovementComponent>())
+		{
+			if (USkeletalMeshComponent* SkeletalMeshComponent = Actor->GetComponentByClass<USkeletalMeshComponent>())
+			{
+				// tick order: CharacterMovementComponent(s) -> ChimeraIslandComponent -> SkeletalMeshComponent(s)
+				CharacterMovementComponents.AddUnique(CharacterMovementComponent);
+				SkeletalMeshComponents.AddUnique(SkeletalMeshComponent);
+
+				AddTickPrerequisiteComponent(CharacterMovementComponent);
+				SkeletalMeshComponent->AddTickPrerequisiteComponent(this);
+			}
+		}
+	}
+}
+
+void UChimeraIslandComponent::AddSearchContext(const UE::Chimera::FSearchContext& SearchContext)
+{
+#if DO_CHECK
+	check(SearchContext.IsValid());
+
+	for (const UE::Chimera::FSearchContext& ContainedSearchContext : SearchContexts)
+	{
+		check(!ContainedSearchContext.IsEquivalent(SearchContext));
+	}
+#endif // DO_CHECK
+	SearchContexts.Add(SearchContext);
+}
+
+void UChimeraIslandComponent::ResetSearchContexts()
+{
+	SearchContexts.Reset();
+}
+
+void UChimeraIslandComponent::ResetSearchResults()
+{
+	SearchResults.Reset();
+	bSearchPerfomed = false;
+}
+
+void UChimeraIslandComponent::UninjectFromAllActors()
+{
+	// Called by UChimeraSubsystem::Tick when there aren't animation jobs flying. No need to FScopeLock Lock(&Mutex);
+	for (TWeakObjectPtr<UCharacterMovementComponent>& CharacterMovementComponentPtr : CharacterMovementComponents)
+	{
+		RemoveTickPrerequisiteComponent(CharacterMovementComponentPtr.Get());
+	}
+
+	for (TWeakObjectPtr<USkeletalMeshComponent>& SkeletalMeshComponentPtr : SkeletalMeshComponents)
+	{
+		SkeletalMeshComponentPtr->RemoveTickPrerequisiteComponent(this);
+	}
+
+	CharacterMovementComponents.Reset();
+	SkeletalMeshComponents.Reset();
+}
+
+bool UChimeraIslandComponent::IsUninjected()
+{
+	const bool bIsUninjected = SkeletalMeshComponents.IsEmpty();
+
+#if DO_CHECK
+	if (bIsUninjected)
+	{
+		// making sure properties are consistent within each other
+		check(CharacterMovementComponents.IsEmpty());
+		check(SearchContexts.IsEmpty());
+
+		FScopeLock Lock(&SearchResultsMutex);
+		check(SearchResults.IsEmpty());
+		check(bSearchPerfomed == false);
+	}
+#endif // DO_CHECK
+
+	return bIsUninjected;
+}
+
+bool UChimeraIslandComponent::DoSearch_AnyThread(UObject* AnimInstance, FChimeraBlueprintResult& Result)
+{
+	using namespace UE::Chimera;
+	
+	FScopeLock Lock(&SearchResultsMutex);
+
+	if (!bSearchPerfomed)
+	{
+		FMemMark Mark(FMemStack::Get());
+
+		TArray<UAnimInstance*, TInlineAllocator<UE::PoseSearch::PreallocatedRolesNum, TMemStackAllocator<>>> AnimInstances;
+		TArray<const UE::PoseSearch::IPoseHistory*, TInlineAllocator<UE::PoseSearch::PreallocatedRolesNum, TMemStackAllocator<>>> PoseHistories;
+		TArray<UE::PoseSearch::FSearchResult, TMemStackAllocator<>> PoseSearchResults;
+		PoseSearchResults.SetNum(SearchContexts.Num());
+
+		for (int32 SearchIndex = 0; SearchIndex < SearchContexts.Num(); ++SearchIndex)
+		{
+			FSearchContext& SearchContext = SearchContexts[SearchIndex];
+			const UPoseSearchDatabase* Database = SearchContext.Database.Get();
+			if (!Database)
+			{
+				UE_LOG(LogChimera, Error, TEXT("UChimeraIslandComponent::DoSearch_AnyThread invalid context database"));
+				return false;
+			}
+
+			if (!Database->Schema)
+			{
+				UE_LOG(LogChimera, Error, TEXT("UChimeraIslandComponent::DoSearch_AnyThread invalid schema for context database %s"), *Database->GetName());
+				return false;
+			}
+
+			AnimInstances.Reset();
+			for (TWeakObjectPtr<UAnimInstance>& AnimInstancePtr : SearchContext.AnimInstances)
+			{
+				UAnimInstance* SearchContextAnimInstance = AnimInstancePtr.Get();
+				if (!SearchContextAnimInstance)
+				{
+					UE_LOG(LogChimera, Error, TEXT("UChimeraIslandComponent::DoSearch_AnyThread null anim instance"));
+					return false;
+				}
+
+				AnimInstances.Add(SearchContextAnimInstance);
+			}
+
+			PoseHistories.Reset();
+			for (const FAnimNode_PoseSearchHistoryCollector_Base* HistoryCollector : SearchContext.HistoryCollectors)
+			{
+				PoseHistories.Add(&HistoryCollector->GetPoseHistory());
+			}
+
+			const UObject* AssetsToSearch[] = { Database };
+			FPoseSearchFutureProperties PoseSearchFutureProperties;
+			const UE::PoseSearch::FSearchResult PoseSearchResult = UPoseSearchLibrary::MotionMatch(AnimInstances, SearchContext.Roles,
+				PoseHistories, AssetsToSearch, SearchContext.ContinuingProperties, PoseSearchFutureProperties);
+
+			if (PoseSearchResult.PoseCost.GetTotalCost() < SearchContext.MaxCost)
+			{
+				PoseSearchResults[SearchIndex] = PoseSearchResult;
+			}
+		}
+
+		////////////////////////////////////////////////////////////////////////////////////////////////////////////////
+		// WIP!
+		// @todo: figure out multiple policies to sue the most characters? right now only the best search is "valid"
+		FPoseSearchCost BestPoseCost;
+		int32 BestSearchIndex = INDEX_NONE;
+		for (int32 SearchIndex = 0; SearchIndex < PoseSearchResults.Num(); ++SearchIndex)
+		{
+			const UE::PoseSearch::FSearchResult& PoseSearchResult = PoseSearchResults[SearchIndex];
+			if (PoseSearchResult.IsValid() && PoseSearchResult.PoseCost < BestPoseCost)
+			{
+				BestPoseCost = PoseSearchResult.PoseCost;
+				BestSearchIndex = SearchIndex;
+			}
+		}
+
+		if (BestSearchIndex != INDEX_NONE)
+		{
+			SearchResults.SetNum(1);
+			SearchResults[0] = InitSearchResult(PoseSearchResults[BestSearchIndex], BestSearchIndex, SearchContexts[BestSearchIndex]);
+		}
+		////////////////////////////////////////////////////////////////////////////////////////////////////////////////
+
+		bSearchPerfomed = true;
+	}
+
+	// looking for AnimInstance in SearchResults to fill up Result
+	for (const FSearchResult& SearchResult : SearchResults)
+	{
+		const FSearchContext& SearchContext = SearchContexts[SearchResult.SearchIndex];
+		for (int32 AnimInstanceIndex = 0; AnimInstanceIndex < SearchContext.AnimInstances.Num(); ++AnimInstanceIndex)
+		{
+			if (SearchContext.AnimInstances[AnimInstanceIndex].Get() == AnimInstance)
+			{
+				const UPoseSearchDatabase* Database = SearchResult.Database.Get();
+				check(Database);
+
+				const UE::PoseSearch::FSearchIndexAsset* SearchIndexAsset = SearchResult.GetSearchIndexAsset();
+				check(SearchIndexAsset);
+
+				const FPoseSearchDatabaseAnimationAssetBase* DatabaseAnimationAssetBase = Database->GetDatabaseAnimationAsset<FPoseSearchDatabaseAnimationAssetBase>(*SearchIndexAsset);
+				check(DatabaseAnimationAssetBase);
+
+				const UE::PoseSearch::FRole Role = SearchContext.Roles[AnimInstanceIndex];
+
+				Result.SelectedAnimation = DatabaseAnimationAssetBase->GetAnimationAsset();
+				Result.SelectedTime = SearchResult.AssetTime;
+				Result.bIsContinuingPoseSearch = SearchResult.bIsContinuingPoseSearch;
+				Result.bLoop = SearchIndexAsset->IsLooping();
+				Result.bIsMirrored = SearchIndexAsset->IsMirrored();
+				Result.BlendParameters = SearchIndexAsset->GetBlendParameters();
+				Result.SelectedDatabase = Database;
+				Result.SearchCost = SearchResult.PoseCost.GetTotalCost();
+				Result.Role = Role;
+
+				if (SearchResult.FullAlignedActorRootBoneTransforms.IsValidIndex(AnimInstanceIndex))
+				{
+					Result.FullAlignedActorRootBoneTransform = SearchResult.FullAlignedActorRootBoneTransforms[AnimInstanceIndex];
+				}
+
+				//Result.bIsFromContinuingPlaying = SearchResult.bIsContinuingPoseSearch;
+
+				// figuring out the WantedPlayRate
+				Result.WantedPlayRate = 1.f;
+				//if (Future.Animation && Future.IntervalTime > 0.f)
+				//{
+				//	if (const UPoseSearchFeatureChannel_PermutationTime* PermutationTimeChannel = Database->Schema->FindFirstChannelOfType<UPoseSearchFeatureChannel_PermutationTime>())
+				//	{
+				//		const FSearchIndex& SearchIndex = Database->GetSearchIndex();
+				//		if (!SearchIndex.IsValuesEmpty())
+				//		{
+				//			TConstArrayView<float> ResultData = Database->GetSearchIndex().GetPoseValues(SearchResult.PoseIdx);
+				//			const float ActualIntervalTime = PermutationTimeChannel->GetPermutationTime(ResultData);
+				//			ProviderResult.WantedPlayRate = ActualIntervalTime / Future.IntervalTime;
+				//		}
+				//	}
+				//}
+
+				// we found our AnimInstance: we can stop searching
+				return true;
+			}
+		}
+	}
+
+	return false;
+}
+
+const UE::Chimera::FSearchResult* UChimeraIslandComponent::FindSearchResult(const UE::Chimera::FSearchContext& SearchContext) const
+{
+	using namespace UE::Chimera;
+
+	// searching for InSearchContext in all the SearchContexts referenced by valid active SearchResults
+	for (const FSearchResult& SearchResult : SearchResults)
+	{
+		const FSearchContext& LocalSearchContext = SearchContexts[SearchResult.SearchIndex];
+		if (LocalSearchContext.Database == SearchContext.Database &&
+			LocalSearchContext.AnimInstances == SearchContext.AnimInstances &&
+			LocalSearchContext.Roles == SearchContext.Roles)
+		{
+			return &SearchResult;
+		}
+	}
+	return nullptr;
+}
diff --git a/Engine/Plugins/Experimental/Animation/Chimera/Source/Runtime/Private/ChimeraLibrary.cpp b/Engine/Plugins/Experimental/Animation/Chimera/Source/Runtime/Private/ChimeraLibrary.cpp
new file mode 100644
--- /dev/null
+++ b/Engine/Plugins/Experimental/Animation/Chimera/Source/Runtime/Private/ChimeraLibrary.cpp
@@ -0,0 +1,25 @@
+// Copyright Epic Games, Inc. All Rights Reserved.
+
+#include "Chimera/ChimeraLibrary.h"
+#include "Chimera/ChimeraSubsystem.h"
+#include "Engine/World.h"
+
+FChimeraBlueprintResult UChimeraLibrary::ChimeraQuery_Pure(TArray<FChimeraAvailability> Availabilities, UObject* AnimInstance, FName PoseHistoryName, bool bValidateResultAgainstAvailabilities)
+{
+	FChimeraBlueprintResult Result;
+	if (UChimeraSubsystem* ChimeraSubsystem = UChimeraSubsystem::GetSubsystem_AnyThread(AnimInstance))
+	{
+		ChimeraSubsystem->Query_AnyThread(Availabilities, AnimInstance, Result, PoseHistoryName, nullptr, bValidateResultAgainstAvailabilities);
+	}
+	return Result;
+}
+
+FChimeraBlueprintResult UChimeraLibrary::ChimeraQuery(const TArrayView<const FChimeraAvailability> Availabilities, UObject* AnimInstance, const FAnimNode_PoseSearchHistoryCollector_Base* HistoryCollector, bool bValidateResultAgainstAvailabilities)
+{
+	FChimeraBlueprintResult Result;
+	if (UChimeraSubsystem* ChimeraSubsystem = UChimeraSubsystem::GetSubsystem_AnyThread(AnimInstance))
+	{
+		ChimeraSubsystem->Query_AnyThread(Availabilities, AnimInstance, Result, FName(), HistoryCollector, bValidateResultAgainstAvailabilities);
+	}
+	return Result;
+}
diff --git a/Engine/Plugins/Experimental/Animation/Chimera/Source/Runtime/Private/ChimeraModule.cpp b/Engine/Plugins/Experimental/Animation/Chimera/Source/Runtime/Private/ChimeraModule.cpp
new file mode 100644
--- /dev/null
+++ b/Engine/Plugins/Experimental/Animation/Chimera/Source/Runtime/Private/ChimeraModule.cpp
@@ -0,0 +1,30 @@
+// Copyright Epic Games, Inc. All Rights Reserved.
+
+#include "CoreMinimal.h"
+#include "Modules/ModuleManager.h"
+
+#define LOCTEXT_NAMESPACE "FChimeraRuntimeModule"
+
+namespace UE::Chimera
+{
+
+class FChimeraRuntimeModule : public IModuleInterface
+{
+public:
+	//~IModuleInterface
+	virtual void StartupModule() override
+	{
+	}
+
+	virtual void ShutdownModule() override
+	{
+	}
+	//~End of IModuleInterface
+};
+
+} // namespace UE::Chimera
+
+IMPLEMENT_MODULE(UE::Chimera::FChimeraRuntimeModule, Chimera)
+
+#undef LOCTEXT_NAMESPACE
+
diff --git a/Engine/Plugins/Experimental/Animation/Chimera/Source/Runtime/Private/ChimeraSubsystem.cpp b/Engine/Plugins/Experimental/Animation/Chimera/Source/Runtime/Private/ChimeraSubsystem.cpp
new file mode 100644
--- /dev/null
+++ b/Engine/Plugins/Experimental/Animation/Chimera/Source/Runtime/Private/ChimeraSubsystem.cpp
@@ -0,0 +1,773 @@
+// Copyright Epic Games, Inc. All Rights Reserved.
+
+#include "Chimera/ChimeraSubsystem.h"
+#include "Chimera/ChimeraDefines.h"
+#include "Chimera/ChimeraIslandComponent.h"
+#include "PoseSearch/PoseSearchDatabase.h"
+#include "PoseSearch/PoseSearchRole.h"
+#include "PoseSearch/PoseSearchSchema.h"
+#include "VisualLogger/VisualLogger.h"
+
+namespace UE::Chimera
+{
+	FString FChimeraAvailabilityEx::GetPoseHistoryName() const
+	{
+		if (HistoryCollector)
+		{
+			return "HistoryProvider";
+		}
+		return PoseHistoryName.ToString();
+	}
+	
+	const FAnimNode_PoseSearchHistoryCollector_Base* FChimeraAvailabilityEx::GetHistoryCollector(const UAnimInstance* AnimInstance) const
+	{
+		if (HistoryCollector)
+		{
+			return HistoryCollector;
+		}
+		return UPoseSearchLibrary::FindPoseHistoryNode(PoseHistoryName, AnimInstance);
+	}
+
+	struct FSearchContextGroup
+	{
+		bool Contains(const FSearchContext& SearchContext) const
+		{
+			for (const TWeakObjectPtr<UAnimInstance>& AnimInstance : SearchContext.AnimInstances)
+			{
+				if (AnimInstances.Find(AnimInstance.Get()))
+				{
+					return true;
+				}
+			}
+			return false;
+		}
+
+		void Add(const FSearchContext& SearchContext, int32 SearchContextIndex)
+		{
+			for (const TWeakObjectPtr<UAnimInstance>& AnimInstance : SearchContext.AnimInstances)
+			{
+				if (AnimInstance.Get())
+				{
+					AnimInstances.Add(AnimInstance.Get());
+				}
+			}
+
+			SearchContextsIndices.Add(SearchContextIndex);
+		}
+
+		void Merge(const FSearchContextGroup& SearchContextGroup)
+		{
+			for (const UAnimInstance* AnimInstance : SearchContextGroup.AnimInstances)
+			{
+				check(AnimInstance);
+				AnimInstances.Add(AnimInstance);
+			}
+
+			for (int32 SearchContextsIndex : SearchContextGroup.SearchContextsIndices)
+			{
+				SearchContextsIndices.Add(SearchContextsIndex);
+			}
+		}
+
+		TSet<const UAnimInstance*, DefaultKeyFuncs<const UAnimInstance*>, TInlineSetAllocator<UE::PoseSearch::PreallocatedRolesNum>> AnimInstances;
+		TArray<int32, TInlineAllocator<PreallocatedSearchesNum>> SearchContextsIndices;
+	};
+
+	struct FRoledAnimInstance
+	{
+		UAnimInstance* AnimInstance = nullptr;
+		UE::PoseSearch::FRole Role = UE::PoseSearch::DefaultRole;
+		const FAnimNode_PoseSearchHistoryCollector_Base* HistoryCollector = nullptr;
+		float BroadPhaseRadius = 0.f;
+		float MaxCost = MAX_flt;
+	};
+
+	// @todo: inherit from a TSet if it gets too slow 
+	// sorted array of FRoledAnimInstance(s). Sorted by UAnimInstance to prevent "alternating" tuples between different frames like
+	// ([AnimInstanceA, RoleA], [AnimInstanceB, RoleB]) vs ([AnimInstanceB, RoleB], [AnimInstanceA, RoleA]) 
+	// that ultimately represent the same search, but endign up with different SearchIDs
+	struct FRoledAnimInstances
+	{
+		void AddRoledAnimInstace(UAnimInstance* AnimInstance, UE::PoseSearch::FRole Role, const FAnimNode_PoseSearchHistoryCollector_Base* HistoryCollector, float BroadPhaseRadius, float MaxCost)
+		{
+			check(AnimInstance && HistoryCollector);
+
+			for (FRoledAnimInstance& RoledAnimInstance : RoledAnimInstances)
+			{
+				if (RoledAnimInstance.AnimInstance == AnimInstance && RoledAnimInstance.HistoryCollector == HistoryCollector && RoledAnimInstance.Role == Role)
+				{
+					// we already found the FRoledAnimInstances. Updating the BroadPhaseRadius and MaxCost accordingly
+					RoledAnimInstance.BroadPhaseRadius = FMath::Max(RoledAnimInstance.BroadPhaseRadius, BroadPhaseRadius);
+					RoledAnimInstance.MaxCost = FMath::Min(RoledAnimInstance.MaxCost, MaxCost);
+					return;
+				}
+			}
+
+			// we couldn't find the FRoledAnimInstances. Let's add a new one
+			FRoledAnimInstance& NewRoledAnimInstance = RoledAnimInstances.AddDefaulted_GetRef();
+			NewRoledAnimInstance.AnimInstance = AnimInstance;
+			NewRoledAnimInstance.HistoryCollector = HistoryCollector;
+			NewRoledAnimInstance.Role = Role;
+			NewRoledAnimInstance.BroadPhaseRadius = BroadPhaseRadius;
+			NewRoledAnimInstance.MaxCost = MaxCost;
+
+			// @todo: use a better datastructure to keep things FRoledAnimInstance(s) sorted? Or use insertion sort?
+			RoledAnimInstances.Sort([](const FRoledAnimInstance& A, const FRoledAnimInstance& B)
+			{
+				// we're sorting by AnimInstance address in memory, since searches neeed to be consistent 
+				// between frames on the same machine, not between replicated machines
+				return B.AnimInstance < A.AnimInstance;
+			});
+		}
+
+		TArrayView<const FRoledAnimInstance> GetDataView() const
+		{
+			return MakeArrayView(RoledAnimInstances.GetData(), RoledAnimInstances.Num());
+		}
+
+	private:
+		TArray<FRoledAnimInstance, TInlineAllocator<UE::PoseSearch::PreallocatedRolesNum, TMemStackAllocator<>>> RoledAnimInstances;
+	};
+
+	typedef TMap<const UPoseSearchDatabase*, FRoledAnimInstances, TInlineSetAllocator<PreallocatedSearchesNum, UE::PoseSearch::TMemStackSetAllocator<>>> FDatabaseToRoledAnimInstances;
+	typedef TPair<const UPoseSearchDatabase*, FRoledAnimInstances> FDatabaseToRoledAnimInstancesPair;
+
+	static void ProcessAvailabilityRequestsMap(const FAvailabilityRequestsMap& AvailabilityRequestsMap, UWorld* SubsystemWorld, FDatabaseToRoledAnimInstances& DatabaseToRoledAnimInstances)
+	{
+		for (FAvailabilityRequestsMap::TConstIterator Iter = AvailabilityRequestsMap.CreateConstIterator(); Iter; ++Iter)
+		{
+			// making sure the weak pointer is still valid
+			UAnimInstance* AnimInstance = Cast<UAnimInstance>(Iter->Key.Get());
+			if (!AnimInstance)
+			{
+				UE_LOG(LogChimera, Log, TEXT("ProcessAvailabilityRequestsMap null anim instance. The associated character got removed from the world"));
+				continue;
+			}
+
+			if (!AnimInstance->GetWorld())
+			{
+				UE_LOG(LogChimera, Error, TEXT("ProcessAvailabilityRequestsMap AnimInstance %s is not in a world!"), *AnimInstance->GetName());
+				continue;
+			}
+
+			if (SubsystemWorld != AnimInstance->GetWorld())
+			{
+				UE_LOG(LogChimera, Error, TEXT("ProcessAvailabilityRequestsMap AnimInstance %s is from World %s, and supposed to be in World @s!"), *AnimInstance->GetName(), *AnimInstance->GetWorld()->GetName(), *SubsystemWorld->GetName());
+				continue;
+			}
+
+			for (const FChimeraAvailabilityEx& AvailabilityRequest : Iter->Value)
+			{
+				if (!AvailabilityRequest.Database)
+				{
+					UE_LOG(LogChimera, Log, TEXT("ProcessAvailabilityRequestsMap null AvailabilityRequest.Database"));
+					continue;
+				}
+
+				if (!AvailabilityRequest.Database->Schema)
+				{
+					UE_LOG(LogChimera, Error, TEXT("ProcessAvailabilityRequestsMap null Schema for Database %s"), *AvailabilityRequest.Database->GetName());
+					continue;
+				}
+
+				const FAnimNode_PoseSearchHistoryCollector_Base* HistoryCollector = AvailabilityRequest.GetHistoryCollector(AnimInstance);
+				if (!HistoryCollector)
+				{
+					UE_LOG(LogChimera, Error, TEXT("ProcessAvailabilityRequestsMap couldn't find PoseHistory %s for AnimInstance %s"), *AvailabilityRequest.GetPoseHistoryName(), *AnimInstance->GetName());
+					continue;
+				}
+
+				FRoledAnimInstances& RoledAnimInstances = DatabaseToRoledAnimInstances.FindOrAdd(AvailabilityRequest.Database);
+				if (AvailabilityRequest.RolesFilter.IsEmpty())
+				{
+					// adding ALL the possible roles from the database:
+					for (const FPoseSearchRoledSkeleton& RoledSkeleton : AvailabilityRequest.Database->Schema->GetRoledSkeletons())
+					{
+						RoledAnimInstances.AddRoledAnimInstace(AnimInstance, RoledSkeleton.Role, HistoryCollector, AvailabilityRequest.BroadPhaseRadius, AvailabilityRequest.MaxCost);
+					}
+				}
+				else
+				{
+					for (const UE::PoseSearch::FRole& Role : AvailabilityRequest.RolesFilter)
+					{
+						if (AvailabilityRequest.Database->Schema->GetRoledSkeleton(Role))
+						{
+							RoledAnimInstances.AddRoledAnimInstace(AnimInstance, Role, HistoryCollector, AvailabilityRequest.BroadPhaseRadius, AvailabilityRequest.MaxCost);
+						}
+						else
+						{
+							UE_LOG(LogChimera, Warning, TEXT("ProcessAvailabilityRequestsMap unsupported Role %s for Database %s"), *Role.ToString(), *AvailabilityRequest.Database->GetName());
+						}
+					}
+				}
+			}
+		}
+	}
+
+	static void PopulateContinuingProperties(FSearchContext& SearchContext, const FIslands& Islands, float DeltaSeconds)
+	{
+		// searching this SearchContext in all the islands to initialize its continuing pose
+		for (const UChimeraIslandComponent* Island : Islands)
+		{
+			if (const FSearchResult* SearchResult = Island->FindSearchResult(SearchContext))
+			{
+				// is still valid...
+				if (SearchResult->IsValid())
+				{
+					if (const UE::PoseSearch::FSearchIndexAsset* SearchIndexAsset = SearchResult->GetSearchIndexAsset())
+					{
+						if (const FPoseSearchDatabaseAnimationAssetBase* DatabaseAsset = SearchResult->Database->GetDatabaseAnimationAsset<FPoseSearchDatabaseAnimationAssetBase>(*SearchIndexAsset))
+						{
+							SearchContext.ContinuingProperties.PlayingAsset = DatabaseAsset->GetAnimationAsset();
+							SearchContext.ContinuingProperties.PlayingAssetAccumulatedTime = SearchResult->AssetTime + DeltaSeconds;
+						}
+					}
+				}
+				break;
+			}
+		}
+	}
+
+	template <typename DataType, typename EvaluateCombinationType>
+	static void GeneratePermutationsRecursive(const TArrayView<DataType> Data, int32 DataIndex, TArrayView<int32> Combination, int32 CombinationIndex, EvaluateCombinationType EvaluateCombination)
+	{
+		if (CombinationIndex == Combination.Num())
+		{ 
+			EvaluateCombination(Data, Combination);
+		}
+		else if (DataIndex < Data.Num())
+		{
+			Combination[CombinationIndex] = DataIndex;
+			GeneratePermutationsRecursive(Data, DataIndex + 1, Combination, CombinationIndex + 1, EvaluateCombination);
+			GeneratePermutationsRecursive(Data, DataIndex + 1, Combination, CombinationIndex, EvaluateCombination);
+		}
+	} 
+
+	template <typename DataType, typename EvaluateCombinationType>
+	static void GeneratePermutations(const TArrayView<DataType> Data, int32 CombinationCardinality, EvaluateCombinationType EvaluateCombination)
+	{
+		TArray<int32, TInlineAllocator<UE::PoseSearch::PreallocatedRolesNum>> Combination;
+		Combination.SetNum(CombinationCardinality);
+		GeneratePermutationsRecursive(Data, 0, Combination, 0, EvaluateCombination);
+	}
+}
+
+FCriticalSection UChimeraSubsystem::RetrieveSubsystemMutex;
+
+UChimeraIslandComponent& UChimeraSubsystem::CreateIsland()
+{
+	UChimeraIslandComponent& Island = *Islands.Add_GetRef(NewObject<UChimeraIslandComponent>()).Get();
+	Island.PrimaryComponentTick.bCanEverTick = true;
+	Island.PrimaryComponentTick.bStartWithTickEnabled = true;
+	Island.SetTickGroup(ETickingGroup::TG_PrePhysics);
+	Island.AddToRoot();
+	Island.RegisterComponentWithWorld(GetWorld());
+	return Island;
+}
+
+UChimeraIslandComponent& UChimeraSubsystem::GetAvailableIsland()
+{
+	for (TObjectPtr<UChimeraIslandComponent>& IslandPtr : Islands)
+	{
+		UChimeraIslandComponent& Island = *IslandPtr.Get();
+		if (Island.IsUninjected())
+		{
+			return Island;
+		}
+	}
+
+	return CreateIsland();
+}
+
+void UChimeraSubsystem::DestroyIsland(int32 Index)
+{
+	UChimeraIslandComponent& Island = *Islands[Index].Get();
+	Island.UninjectFromAllActors();
+	Island.RemoveFromRoot();
+	Island.UnregisterComponent();
+	Island.DestroyComponent();
+	Islands.RemoveAt(Index);
+}
+
+void UChimeraSubsystem::DestroyAllIslands()
+{
+	for (int32 IslandIndex = Islands.Num() - 1; IslandIndex >= 0; --IslandIndex)
+	{
+		DestroyIsland(IslandIndex);
+	}
+}
+
+void UChimeraSubsystem::UninjectAllIslands()
+{
+	for (TObjectPtr<UChimeraIslandComponent>& IslandPtr : Islands)
+	{
+		UChimeraIslandComponent& Island = *IslandPtr.Get();
+		Island.UninjectFromAllActors();
+		Island.ResetSearchContexts();
+		Island.ResetSearchResults();
+	}
+}
+
+bool UChimeraSubsystem::ValidateAllIslands() const
+{
+#if DO_CHECK
+
+	TSet<TWeakObjectPtr<UCharacterMovementComponent>> CharacterMovementComponents;
+	TSet<TWeakObjectPtr<USkeletalMeshComponent>> SkeletalMeshComponents;
+
+	bool bAlreadyInSet = false;
+	for (const TObjectPtr<UChimeraIslandComponent>& Island : Islands)
+	{
+		for (const TWeakObjectPtr<UCharacterMovementComponent>& CharacterMovementComponent : Island->GetCharacterMovementComponents())
+		{
+			CharacterMovementComponents.Add(CharacterMovementComponent, &bAlreadyInSet);
+			if (bAlreadyInSet)
+			{
+				return false;
+			}
+		}
+
+		for (const TWeakObjectPtr<USkeletalMeshComponent>& SkeletalMeshComponent : Island->GetSkeletalMeshComponents())
+		{
+			SkeletalMeshComponents.Add(SkeletalMeshComponent, &bAlreadyInSet);
+			if (bAlreadyInSet)
+			{
+				return false;
+			}
+		}
+	}
+#endif // DO_CHECK
+	return true;
+}
+
+UChimeraIslandComponent* UChimeraSubsystem::FindIsland(UObject* InAnimInstance)
+{
+	if (InAnimInstance)
+	{
+		if (UAnimInstance* AnimInstance = Cast<UAnimInstance>(InAnimInstance))
+		{
+			if (AActor* Actor = AnimInstance->GetOwningActor())
+			{
+				if (USkeletalMeshComponent* SkeletalMeshComponent = Actor->GetComponentByClass<USkeletalMeshComponent>())
+				{
+					for (const TObjectPtr<UChimeraIslandComponent>& Island : Islands)
+					{
+						for (const TWeakObjectPtr<USkeletalMeshComponent>& IslandSkeletalMeshComponent : Island->GetSkeletalMeshComponents())
+						{
+							if (IslandSkeletalMeshComponent.Get() == SkeletalMeshComponent)
+							{
+								return Island.Get();
+							}
+						}
+					}
+				}
+			}
+		}
+	}
+	return nullptr;
+}
+
+UChimeraSubsystem* UChimeraSubsystem::GetSubsystem_AnyThread(UObject* AnimInstance)
+{
+	if (AnimInstance)
+	{
+		if (UWorld* World = AnimInstance->GetWorld())
+		{
+			FScopeLock Lock(&RetrieveSubsystemMutex);
+			return UWorld::GetSubsystem<UChimeraSubsystem>(World);
+		}
+	}
+	return nullptr;
+}
+
+void UChimeraSubsystem::Deinitialize()
+{
+	DestroyAllIslands();
+
+	Super::Deinitialize();
+}
+
+void UChimeraSubsystem::Tick(float DeltaSeconds)
+{
+	using namespace UE::Chimera;
+
+	check(IsInGameThread());
+	check(GetWorld());
+
+	Super::Tick(DeltaSeconds);
+
+	// FScopeLock Lock(&AvailabilityRequestsMapMutex); is not necessary since UChimeraSubsystem gets ticked outside the parallel animation jobs
+	FMemMark Mark(FMemStack::Get());
+	
+	// processing the AvailabilityRequestsMap to fill up a map of Databases pointing to an array of all the AnimInstances with related roles
+	FDatabaseToRoledAnimInstances DatabaseToRoledAnimInstances;
+	ProcessAvailabilityRequestsMap(AvailabilityRequestsMap, GetWorld(), DatabaseToRoledAnimInstances);
+
+	// for each database now we try to create all the possible permutations of the roled anim instances
+	// for example, given a database set up with assets for 2 characters interactions with roles RoleA and RoleB
+	// and 2 anim instances, all of them willing to partecipate in the 2 characters interaction with both roles RoleA and RoleB:
+	// CharA could be taking RoleA and RoleB,
+	// CharB could be taking RoleA and RoleB,
+	// we generate all the permutations from the array of options:
+	// CharA/RoleA, CharA/RoleB, CharB/RoleA, CharB/RoleB
+	//
+	// and we prune the invalid tuples:
+	//
+	// CharA/RoleA - CharA/RoleB -> invalid because of same CharA
+	// CharA/RoleA - CharB/RoleA -> invalid because of same RoleA
+	// CharA/RoleA - CharB/RoleB -> VALID!
+	//
+	// CharA/RoleB - CharB/RoleA -> VALID!
+	// CharA/RoleB - CharB/RoleB -> invalid because of same RoleB
+	//
+	// CharB/RoleA - CharB/RoleB -> invalid because of same CharB
+
+	TSet<const UAnimInstance*, DefaultKeyFuncs<const UAnimInstance*>, TInlineSetAllocator<UE::PoseSearch::PreallocatedRolesNum>> AnimInstances;
+	TSet<UE::PoseSearch::FRole, DefaultKeyFuncs<UE::PoseSearch::FRole>, TInlineSetAllocator<UE::PoseSearch::PreallocatedRolesNum>> CoveredRoles;
+	TArray<FSearchContext, TInlineAllocator<PreallocatedSearchesNum, TMemStackAllocator<>>> SearchContexts;
+	for (const FDatabaseToRoledAnimInstancesPair& DatabaseToRoledAnimInstancesPair : DatabaseToRoledAnimInstances)
+	{
+		const UPoseSearchDatabase* Database = DatabaseToRoledAnimInstancesPair.Key;
+		const TArray<FPoseSearchRoledSkeleton>& RoledSkeletons = Database->Schema->GetRoledSkeletons();
+		const int32 CombinationCardinality = RoledSkeletons.Num();
+		
+		// @todo: naive implementation: optimize this code!
+		GeneratePermutations(DatabaseToRoledAnimInstancesPair.Value.GetDataView(), CombinationCardinality, 
+			[&AnimInstances, &CoveredRoles, Database, &RoledSkeletons, &SearchContexts](const TArrayView<const FRoledAnimInstance> RoledAnimInstances, const TArrayView<int32> Combination)
+			{
+				// checking if this combination is valid for the Database:
+				AnimInstances.Reset();
+				CoveredRoles.Reset();
+
+				for (int32 CombinationIndex = 0; CombinationIndex < Combination.Num(); ++CombinationIndex)
+				{
+					const FRoledAnimInstance& RoledAnimInstance = RoledAnimInstances[Combination[CombinationIndex]];
+
+					bool bIsAlreadyInSet = false;
+					AnimInstances.Add(RoledAnimInstance.AnimInstance, &bIsAlreadyInSet);
+
+					// is there any duplicated AnimInstance?
+					if (bIsAlreadyInSet)
+					{
+						return false;
+					}
+
+					CoveredRoles.Add(RoledAnimInstance.Role);
+				}
+
+				// does it cover all the roles?
+				for (const FPoseSearchRoledSkeleton& RoledSkeleton : RoledSkeletons)
+				{
+					if (!CoveredRoles.Find(RoledSkeleton.Role))
+					{
+						return false;
+					}
+				}
+
+				// @todo: naive implementation: optimize this code!
+				// checking if all the actors are within BroadPhaseRadius to each other
+				for (int32 CombinationIndexA = 0; CombinationIndexA < Combination.Num(); ++CombinationIndexA)
+				{
+					const FRoledAnimInstance& RoledAnimInstanceA = RoledAnimInstances[Combination[CombinationIndexA]];
+					const FVector ActorLocationA = RoledAnimInstanceA.AnimInstance->GetOwningActor()->GetActorLocation();
+
+					for (int32 CombinationIndexB = 0; CombinationIndexB < Combination.Num(); ++CombinationIndexB)
+					{
+						if (CombinationIndexA != CombinationIndexB)
+						{
+							const FRoledAnimInstance& RoledAnimInstanceB = RoledAnimInstances[Combination[CombinationIndexB]];
+							const FVector ActorLocationB = RoledAnimInstanceB.AnimInstance->GetOwningActor()->GetActorLocation();
+
+							const FVector DeltaLocation = ActorLocationA - ActorLocationB;
+							if (DeltaLocation.Length() > RoledAnimInstanceA.BroadPhaseRadius)
+							{
+								return false;
+							}
+						}
+					}
+				}
+
+				FSearchContext& SearchContext = SearchContexts.AddDefaulted_GetRef();
+				SearchContext.Database = Database;
+				SearchContext.AnimInstances.Reserve(Combination.Num());
+				SearchContext.HistoryCollectors.Reserve(Combination.Num());
+				SearchContext.Roles.Reserve(Combination.Num());
+				
+#if ENABLE_DRAW_DEBUG
+				SearchContext.BroadPhaseRadiuses.Reserve(Combination.Num());
+#endif //ENABLE_DRAW_DEBUG
+
+				for (int32 CombinationIndex = 0; CombinationIndex < Combination.Num(); ++CombinationIndex)
+				{
+					const FRoledAnimInstance& RoledAnimInstance = RoledAnimInstances[Combination[CombinationIndex]];
+					SearchContext.AnimInstances.Add(RoledAnimInstance.AnimInstance);
+					SearchContext.HistoryCollectors.Add(RoledAnimInstance.HistoryCollector);
+					SearchContext.Roles.Add(RoledAnimInstance.Role);
+#if ENABLE_DRAW_DEBUG
+					SearchContext.BroadPhaseRadiuses.Add(RoledAnimInstance.BroadPhaseRadius);
+#endif //ENABLE_DRAW_DEBUG
+
+					// search is valid only if cost result is lower than any AvailabilityRequest.MaxCost (copied in RoledAnimInstance.MaxCost)
+					SearchContext.MaxCost = FMath::Min(SearchContext.MaxCost, RoledAnimInstance.MaxCost);
+				}
+
+				return true;
+			});
+	}
+	
+	// for each valid SearchContexts we try to figure out the continuing pose properties from the current Islands
+	for (FSearchContext& SearchContext : SearchContexts)
+	{
+		PopulateContinuingProperties(SearchContext, Islands, DeltaSeconds);
+	}
+
+	// grouping SearchContexts in different Islands
+	TArray<FSearchContextGroup, TInlineAllocator<PreallocatedSearchesNum, TMemStackAllocator<>>> SearchContextGroups;
+	for (int32 SearchContextIndex = 0; SearchContextIndex < SearchContexts.Num(); ++SearchContextIndex)
+	{
+		const FSearchContext& SearchContext = SearchContexts[SearchContextIndex];
+
+		int32 MainSearchContextGroupIndex = INDEX_NONE;
+		for (int32 SearchContextGroupIndex = 0; SearchContextGroupIndex < SearchContextGroups.Num();)
+		{
+			if (SearchContextGroups[SearchContextGroupIndex].Contains(SearchContext))
+			{
+				if (MainSearchContextGroupIndex == INDEX_NONE)
+				{
+					MainSearchContextGroupIndex = SearchContextGroupIndex;
+					SearchContextGroups[MainSearchContextGroupIndex].Add(SearchContext, SearchContextIndex);
+					++SearchContextGroupIndex;
+				}
+				else
+				{
+					SearchContextGroups[MainSearchContextGroupIndex].Merge(SearchContextGroups[SearchContextGroupIndex]);
+					SearchContextGroups.RemoveAt(SearchContextGroupIndex);
+				}
+			}
+			else
+			{
+				++SearchContextGroupIndex;
+			}
+		}
+		if (MainSearchContextGroupIndex == INDEX_NONE)
+		{
+			SearchContextGroups.AddDefaulted_GetRef().Add(SearchContext, SearchContextIndex);
+		}
+	}
+
+#if ENABLE_DRAW_DEBUG
+	DebugDraw();
+#endif //ENABLE_DRAW_DEBUG
+
+	UninjectAllIslands();
+
+	for (const FSearchContextGroup& SearchContextGroup : SearchContextGroups)
+	{
+		UChimeraIslandComponent& Island = GetAvailableIsland();
+
+		for (const UAnimInstance* AnimInstance : SearchContextGroup.AnimInstances)
+		{
+			Island.InjectToActor(AnimInstance->GetOwningActor());
+		}
+
+		for (int32 SearchContextsIndex : SearchContextGroup.SearchContextsIndices)
+		{
+			Island.AddSearchContext(SearchContexts[SearchContextsIndex]);
+		}
+	}
+
+	AvailabilityRequestsMap.Reset();
+
+	check(ValidateAllIslands());
+}
+
+TStatId UChimeraSubsystem::GetStatId() const
+{
+	RETURN_QUICK_DECLARE_CYCLE_STAT(UChimeraSubsystem, STATGROUP_Tickables);
+}
+
+void UChimeraSubsystem::Query_AnyThread(const TArrayView<const FChimeraAvailability> Availabilities, UObject* AnimInstance, FChimeraBlueprintResult& Result,
+	FName PoseHistoryName, const FAnimNode_PoseSearchHistoryCollector_Base* HistoryCollector, bool bValidateResultAgainstAvailabilities)
+{
+	using namespace UE::Chimera;
+
+	Result = FChimeraBlueprintResult();
+
+	if (AnimInstance)
+	{
+		// if we find AnimInstance in an island, we perform ALL the Island motion matching searches.
+		// thread safety is ensured by the lock on Island->SearchResultsMutex in DoSearch_AnyThread
+		if (UChimeraIslandComponent* Island = FindIsland(AnimInstance))
+		{
+			Island->DoSearch_AnyThread(AnimInstance, Result);
+
+			if (bValidateResultAgainstAvailabilities && Result.SelectedAnimation)
+			{
+				bool bResultValidated = false;
+
+				for (const FChimeraAvailability& Availability : Availabilities)
+				{
+					if (Availability.Database == Result.SelectedDatabase &&
+						(Availability.RolesFilter.IsEmpty() || Availability.RolesFilter.Contains(Result.Role)))
+					{
+						bResultValidated = true;
+						break;
+					}
+				}
+
+				if (!bResultValidated)
+				{
+					Result = FChimeraBlueprintResult();
+				}
+			}
+		}
+
+		// queuing the availabilities for the next frame Query_AnyThread
+		if (!Availabilities.IsEmpty())
+		{
+			FScopeLock Lock(&AvailabilityRequestsMapMutex);
+
+			TArray<FChimeraAvailabilityEx>& AvailabilityRequests = AvailabilityRequestsMap.FindOrAdd(AnimInstance);
+			for (const FChimeraAvailability& Availability : Availabilities)
+			{
+				if (Availability.Database)
+				{
+					AvailabilityRequests.Emplace(Availability, PoseHistoryName, HistoryCollector);
+				}
+			}
+		}
+	}
+}
+
+void UChimeraSubsystem::DebugDraw() const
+{
+	using namespace UE::Chimera;
+
+#if ENABLE_DRAW_DEBUG
+	static const FColor Colors[] = { FColor::Red, FColor::Green, FColor::Blue, FColor::Yellow, FColor::Black };
+	static const int32 NumColors = sizeof(Colors) / sizeof(Colors[0]);
+	int32 CurrentColorIndex = 0;
+
+	for (const TObjectPtr<UChimeraIslandComponent>& IslandPtr : Islands)
+	{
+		UChimeraIslandComponent& Island = *IslandPtr.Get();
+		if (!Island.IsUninjected())
+		{
+			Island.DebugDraw(Colors[CurrentColorIndex]);
+			CurrentColorIndex = (CurrentColorIndex + 1) % NumColors;
+		}
+	}
+
+#endif //ENABLE_DRAW_DEBUG
+
+#if ENABLE_VISUAL_LOG
+	TStringBuilder<512> StringBuilder;
+
+	for (FAvailabilityRequestsMap::TConstIterator Iter = AvailabilityRequestsMap.CreateConstIterator(); Iter; ++Iter)
+	{
+		// making sure the weak pointer is still valid
+		if (UAnimInstance* AnimInstance = Cast<UAnimInstance>(Iter->Key.Get()))
+		{
+			// looking for valid results for this AnimInstance
+			const UPoseSearchDatabase* ResultDatabase = nullptr;
+			UE::PoseSearch::FRole ResultRole;
+			for (const TObjectPtr<UChimeraIslandComponent>& IslandPtr : Islands)
+			{
+				UChimeraIslandComponent& Island = *IslandPtr.Get();
+				if (!Island.IsUninjected())
+				{
+					for (const FSearchResult& SearchResult : Island.GetSearchResults())
+					{
+						const FSearchContext& SearchContext = Island.GetSearchContexts()[SearchResult.SearchIndex];
+						for (int32 AnimInstanceIndex = 0; AnimInstanceIndex < SearchContext.AnimInstances.Num(); ++AnimInstanceIndex)
+						{
+							if (SearchContext.AnimInstances[AnimInstanceIndex].Get() == AnimInstance)
+							{
+								ResultDatabase = SearchResult.Database.Get();
+								check(ResultDatabase);
+								ResultRole = SearchContext.Roles[AnimInstanceIndex];
+								break;
+							}
+						}
+
+						if (ResultDatabase)
+						{
+							break;
+						}
+					}
+				}
+
+				if (ResultDatabase)
+				{
+					break;
+				}
+			}
+
+			StringBuilder.Reset();
+			for (const FChimeraAvailabilityEx& AvailabilityRequest : Iter->Value)
+			{
+				StringBuilder.Append(GetNameSafe(AvailabilityRequest.Database));
+				StringBuilder.Append(" / ");
+				StringBuilder.Append(AvailabilityRequest.GetPoseHistoryName());
+				StringBuilder.Append(" [");
+				bool bAddComma = false;
+
+				if (AvailabilityRequest.RolesFilter.IsEmpty())
+				{
+					// adding ALL the possible roles from the database:
+					if (AvailabilityRequest.Database && AvailabilityRequest.Database->Schema)
+					{
+						for (const FPoseSearchRoledSkeleton& RoledSkeleton : AvailabilityRequest.Database->Schema->GetRoledSkeletons())
+						{
+							if (bAddComma)
+							{
+								StringBuilder.Append(",");
+							}
+							else
+							{
+								bAddComma = true;
+							}
+
+							if (AvailabilityRequest.Database == ResultDatabase && RoledSkeleton.Role == ResultRole)
+							{
+								StringBuilder.Append("(*)");
+							}
+
+							StringBuilder.Append(RoledSkeleton.Role.ToString());
+						}
+					}
+				}
+				else
+				{
+					for (const FName& Role : AvailabilityRequest.RolesFilter)
+					{
+						if (bAddComma)
+						{
+							StringBuilder.Append(",");
+						}
+						else
+						{
+							bAddComma = true;
+						}
+
+						if (AvailabilityRequest.Database == ResultDatabase && Role == ResultRole)
+						{
+							StringBuilder.Append("(*)");
+						}
+
+						StringBuilder.Append(Role.ToString());
+					}
+				}
+
+				StringBuilder.Append("]");
+				StringBuilder.Append("\n");
+			}
+			
+			// @todo: waiting for visual logger fix to display strings :/
+			const FVector StringOffset(0.f, 0.f, 0.001f);
+			const FVector ActorLocation = AnimInstance->GetOwningActor()->GetActorLocation();
+			UE_VLOG_SEGMENT(AnimInstance, "Chimera", Display, ActorLocation, ActorLocation + StringOffset, FColor::Transparent, TEXT("%s"), StringBuilder.ToString());
+		}
+	}
+#endif // ENABLE_VISUAL_LOG
+}
\ No newline at end of file
diff --git a/Engine/Plugins/Experimental/Animation/Chimera/Source/Runtime/Public/Chimera/AnimNode_Chimera.h b/Engine/Plugins/Experimental/Animation/Chimera/Source/Runtime/Public/Chimera/AnimNode_Chimera.h
new file mode 100644
--- /dev/null
+++ b/Engine/Plugins/Experimental/Animation/Chimera/Source/Runtime/Public/Chimera/AnimNode_Chimera.h
@@ -0,0 +1,71 @@
+// Copyright Epic Games, Inc. All Rights Reserved.
+
+#pragma once
+
+#include "BlendStack/AnimNode_BlendStack.h"
+#include "Chimera/ChimeraLibrary.h"
+#include "AnimNode_Chimera.generated.h"
+
+struct FAnimationInitializeContext;
+struct FNodeDebugData;
+
+USTRUCT(BlueprintInternalUseOnly, Experimental)
+struct CHIMERA_API FAnimNode_Chimera : public FAnimNode_BlendStack_Standalone
+{
+	GENERATED_BODY();
+
+	UPROPERTY(EditAnywhere, BlueprintReadWrite, Category = Links, meta = (DisplayPriority = 0))
+	FPoseLink Source;
+
+	UPROPERTY(EditAnywhere, Category = Chimera, meta = (PinHiddenByDefault))
+	TArray<FChimeraAvailability> Availabilities;
+
+	UPROPERTY(EditAnywhere, Category = Chimera, meta = (PinHiddenByDefault))
+	bool bValidateResultAgainstAvailabilities = true;
+
+	// time from the beginning of the interaction to warp to full translation alignment
+	UPROPERTY(EditAnywhere, BlueprintReadWrite, Category = Warping, meta = (PinHiddenByDefault, ClampMin = "0"))
+	float InitialTranslationWarpTime = 0.2f;
+
+	// time from the beginning of the interaction to warp to full rotation alignment
+	UPROPERTY(EditAnywhere, BlueprintReadWrite, Category = Warping, meta = (PinHiddenByDefault, ClampMin = "0"))
+	float InitialRotationWarpTime = 0.2f;
+
+	// tunable animation transition blend time 
+	UPROPERTY(EditAnywhere, BlueprintReadWrite, Category = Blending, meta = (PinHiddenByDefault, ClampMin = "0"))
+	float BlendTime = 0.2f;
+
+	UPROPERTY(EditAnywhere, BlueprintReadWrite, Category = Blending, meta = (PinHiddenByDefault, UseAsBlendProfile = true))
+	TObjectPtr<UBlendProfile> BlendProfile;
+
+	UPROPERTY(EditAnywhere, BlueprintReadWrite, Category = Blending, meta = (PinHiddenByDefault))
+	EAlphaBlendOption BlendOption = EAlphaBlendOption::Linear;
+
+	// tunable animation transition blend time 
+	UPROPERTY(EditAnywhere, Category = Blending)
+	bool bUseInertialBlend = false;
+
+	// Reset the blend stack if it has become relevant to the graph after not being updated on previous frames.
+	UPROPERTY(EditAnywhere, Category = Blending)
+	bool bResetOnBecomingRelevant = true;
+
+	// FAnimNode_Base interface
+	virtual void GatherDebugData(FNodeDebugData& DebugData) override;
+	virtual void Initialize_AnyThread(const FAnimationInitializeContext& Context) override;
+	virtual void Evaluate_AnyThread(FPoseContext& Output) override;
+	virtual void CacheBones_AnyThread(const FAnimationCacheBonesContext& Context) override;
+
+	virtual void UpdateAssetPlayer(const FAnimationUpdateContext& Context) override;
+	// End of FAnimNode_Base interface
+
+protected:
+	float BlendLerp = 0.f;
+	float TranslationWarpLerp = 0.f;
+	float RotationWarpLerp = 0.f;
+
+	// Update Counter for detecting being relevant
+	FGraphTraversalCounter UpdateCounter;
+
+	// wanted world transform for full aligment interaction
+	FTransform FullAlignedActorRootBoneTransform = FTransform::Identity;
+};
diff --git a/Engine/Plugins/Experimental/Animation/Chimera/Source/Runtime/Public/Chimera/ChimeraAsset.h b/Engine/Plugins/Experimental/Animation/Chimera/Source/Runtime/Public/Chimera/ChimeraAsset.h
new file mode 100644
--- /dev/null
+++ b/Engine/Plugins/Experimental/Animation/Chimera/Source/Runtime/Public/Chimera/ChimeraAsset.h
@@ -0,0 +1,89 @@
+// Copyright Epic Games, Inc. All Rights Reserved.
+
+#pragma once
+
+#include "Engine/DataAsset.h"
+#include "PoseSearch/MultiAnimAsset.h"
+#include "PoseSearch/PoseSearchRole.h"
+#include "ChimeraAsset.generated.h"
+
+USTRUCT(Experimental)
+struct CHIMERA_API FChimeraAssetItem
+{
+	GENERATED_BODY()
+
+	// associated aniamtion for this FChimeraAssetItem
+	UPROPERTY(EditAnywhere, Category = "Settings")
+	TObjectPtr<UAnimationAsset> Animation;
+
+	// associated role for this FChimeraAssetItem
+	UPROPERTY(EditAnywhere, Category = "Settings")
+	FName Role;
+
+	// relative weight to the other FChimeraAssetItem::WarpingWeightRotation(s) defining which character will be rotated while warping
+	// 0 - the associated character to this item will move fully to compensate the warping errors
+	// > 0 && all the other FChimeraAssetItem::WarpingWeightTranslation as zero, and the associated character will not move
+	UPROPERTY(EditAnywhere, Category = "Warping", meta = (ClampMin = "0", ClampMax = "1"))
+	float WarpingWeightRotation = 0.5f;
+
+	// relative weight to the other FChimeraAssetItem::WarpingWeightTranslation(s) defining which character will be translated while warping
+	// 0 - the associated character to this item will move fully to compensate the warping errors
+	// > 0 && all the other FChimeraAssetItem::WarpingWeightTranslation as zero, and the associated character will not move
+	UPROPERTY(EditAnywhere, Category = "Warping", meta = (ClampMin = "0", ClampMax = "1"))
+	float WarpingWeightTranslation = 0.5f;
+
+	// offset from the origin
+	UPROPERTY(EditAnywhere, Category = "Settings")
+	FTransform Origin = FTransform::Identity;
+};
+
+UCLASS(Experimental, BlueprintType, Category = "Animation|Chimera")
+class CHIMERA_API UChimeraAsset : public UMultiAnimAsset
+{
+	GENERATED_BODY()
+
+	UPROPERTY(EditAnywhere, Category = "Settings")
+	TArray<FChimeraAssetItem> Items;
+
+public:
+#if WITH_EDITORONLY_DATA
+
+	UPROPERTY(Transient, EditAnywhere, Category = "Debug", meta = (EditCondition=bEnableDebugWarp, EditConditionHides))
+	TArray<FTransform> DebugWarpOffsets;
+
+	// used to test warping: 0 - no warping applied, 1 - full warping/aligment applied
+	// test warping actors will be offsetted by Items::DebugWarpOffset transforms from the original
+	// UMultiAnimAsset::GetOrigin() definition and warped accordingly with CalculateWarpTransforms
+	// following the rotation and translation weights defined in Items::WarpingWeightRotation and
+	// Items::WarpingWeightTranslation as relative weights between the Items (they'll be normalized at runtime)
+	UPROPERTY(Transient, EditAnywhere, Category = "Debug", meta = (ClampMin = "0", ClampMax = "1", EditCondition=bEnableDebugWarp, EditConditionHides))
+	float DebugWarpAmount = 0.f;
+
+	UPROPERTY(Transient, EditAnywhere, Category = "Debug")
+	bool bEnableDebugWarp = false;
+#endif // WITH_EDITORONLY_DATA
+
+	virtual bool IsLooping() const override;
+	virtual bool HasRootMotion() const override;
+	virtual float GetPlayLength() const override;
+
+#if WITH_EDITOR
+	virtual int32 GetFrameAtTime(float Time) const override;
+#endif // WITH_EDITOR
+
+	virtual int32 GetNumRoles() const override { return Items.Num(); }
+	virtual UE::PoseSearch::FRole GetRole(int32 RoleIndex) const override { return Items[RoleIndex].Role; }
+	
+	virtual UAnimationAsset* GetAnimationAsset(const FName /*UE::PoseSearch::FRole*/& Role) const override;
+
+	virtual FTransform GetOrigin(const UE::PoseSearch::FRole& Role) const override;
+
+#if WITH_EDITOR
+	FTransform GetDebugWarpOrigin(const UE::PoseSearch::FRole& Role, bool bComposeWithDebugWarpOffset) const;
+#endif // WITH_EDITOR
+
+	void CalculateWarpTransforms(float Time, const TArrayView<const FTransform> ActorRootBoneTransforms, TArrayView<FTransform> FullAlignedActorRootBoneTransforms, const UWorld* DebugDrawWorld = nullptr) const;
+
+	FQuat FindReferenceOrientation(const TArrayView<const FTransform> Transforms, const TArrayView<int32> SortedByWarpingWeightRotationItemIndex) const;
+	FVector FindReferencePosition(const TArrayView<const FTransform> Transforms, const TArrayView<float> NormalizedWarpingWeightTranslation) const;
+};
\ No newline at end of file
diff --git a/Engine/Plugins/Experimental/Animation/Chimera/Source/Runtime/Public/Chimera/ChimeraDefines.h b/Engine/Plugins/Experimental/Animation/Chimera/Source/Runtime/Public/Chimera/ChimeraDefines.h
new file mode 100644
--- /dev/null
+++ b/Engine/Plugins/Experimental/Animation/Chimera/Source/Runtime/Public/Chimera/ChimeraDefines.h
@@ -0,0 +1,13 @@
+// Copyright Epic Games, Inc. All Rights Reserved.
+
+#pragma once
+
+#include "Logging/LogMacros.h"
+
+DECLARE_LOG_CATEGORY_EXTERN(LogChimera, Log, All);
+
+namespace UE::Chimera
+{
+	static constexpr int32 PreallocatedSearchesNum = 4;
+
+} // namespace UE::Chimera
diff --git a/Engine/Plugins/Experimental/Animation/Chimera/Source/Runtime/Public/Chimera/ChimeraIslandComponent.h b/Engine/Plugins/Experimental/Animation/Chimera/Source/Runtime/Public/Chimera/ChimeraIslandComponent.h
new file mode 100644
--- /dev/null
+++ b/Engine/Plugins/Experimental/Animation/Chimera/Source/Runtime/Public/Chimera/ChimeraIslandComponent.h
@@ -0,0 +1,86 @@
+// Copyright Epic Games, Inc. All Rights Reserved.
+
+#pragma once
+
+#include "Components/ActorComponent.h"
+#include "PoseSearch/PoseSearchLibrary.h"
+#include "ChimeraIslandComponent.generated.h"
+
+struct FChimeraBlueprintResult;
+struct FPoseSearchContinuingProperties;
+class UCharacterMovementComponent;
+
+namespace UE::Chimera
+{
+	struct FSearchContext
+	{
+		// @todo: since AnimInstances, HistoryCollectors, Roles, and BroadPhaseRadiuses have the same cardinality,
+		// add an array of structs instead of many arrays
+		TArray<TWeakObjectPtr<UAnimInstance>> AnimInstances;
+		// @todo: perhaps use the animnode of the pose history collector instead
+		TArray<const FAnimNode_PoseSearchHistoryCollector_Base*> HistoryCollectors;
+		TArray<UE::PoseSearch::FRole> Roles;
+
+#if ENABLE_DRAW_DEBUG
+		TArray<float> BroadPhaseRadiuses;
+#endif // ENABLE_DRAW_DEBUG
+
+		float MaxCost = MAX_flt;
+
+		TWeakObjectPtr<const UPoseSearchDatabase> Database;
+		FPoseSearchContinuingProperties ContinuingProperties;
+
+		bool IsValid() const;
+		bool IsEquivalent(const FSearchContext& Other) const;
+	};
+
+	struct FSearchResult : public UE::PoseSearch::FSearchResult
+	{
+		int32 SearchIndex = INDEX_NONE;
+		TArray<FTransform, TInlineAllocator<UE::PoseSearch::PreallocatedRolesNum>> FullAlignedActorRootBoneTransforms;
+	};
+} // namespace UE::Chimera
+
+// UChimeraIslandComponent is a component injected between the interacting actors UCharacterMovementComponent and USkeletalMeshComponent
+// to create a execution threading fence to be able to perform motion matching searches between the involved characters in a thread safe manner.
+// Look at UChimeraSubsystem "Execution model and threading details" for additional information
+UCLASS(Experimental, Category="Animation|Chimera")
+class UChimeraIslandComponent : public UActorComponent
+{
+public:
+	GENERATED_BODY()
+
+	virtual void TickComponent(float DeltaTime, enum ELevelTick TickType, FActorComponentTickFunction *ThisTickFunction) override;
+
+	void InjectToActor(AActor* Actor);
+	void UninjectFromAllActors();
+	bool IsUninjected();
+
+	bool DoSearch_AnyThread(UObject* AnimInstance, FChimeraBlueprintResult& Result);
+
+	void AddSearchContext(const UE::Chimera::FSearchContext& SearchContext);
+	void ResetSearchContexts();
+	void ResetSearchResults();
+
+	void DebugDraw(const FColor& Color = FColor::Red) const;
+
+	const TArray<TWeakObjectPtr<UCharacterMovementComponent>>& GetCharacterMovementComponents() const { return CharacterMovementComponents; }
+	const TArray<TWeakObjectPtr<USkeletalMeshComponent>>& GetSkeletalMeshComponents() const { return SkeletalMeshComponents; }
+	const TArray<UE::Chimera::FSearchContext>& GetSearchContexts() const { return SearchContexts; }
+	const TArray<UE::Chimera::FSearchResult>& GetSearchResults() const { return SearchResults; }
+
+	const UE::Chimera::FSearchResult* FindSearchResult(const UE::Chimera::FSearchContext& SearchContext) const;
+
+private:
+	TArray<TWeakObjectPtr<UCharacterMovementComponent>> CharacterMovementComponents;
+	TArray<TWeakObjectPtr<USkeletalMeshComponent>> SkeletalMeshComponents;
+	
+	// there's one FSearchContext for each search we need to perform (including all the possible roles permutations). Added by UChimeraSubsystem::Tick
+	TArray<UE::Chimera::FSearchContext> SearchContexts;
+
+	// SearchResults contains only the best results, and it has not necessarly the same cardinality as SearchContexts. usually SearchResults.Num() < SearchContexts.Num()
+	TArray<UE::Chimera::FSearchResult> SearchResults;
+	bool bSearchPerfomed = false;
+
+	FCriticalSection SearchResultsMutex;
+};
diff --git a/Engine/Plugins/Experimental/Animation/Chimera/Source/Runtime/Public/Chimera/ChimeraLibrary.h b/Engine/Plugins/Experimental/Animation/Chimera/Source/Runtime/Public/Chimera/ChimeraLibrary.h
new file mode 100644
--- /dev/null
+++ b/Engine/Plugins/Experimental/Animation/Chimera/Source/Runtime/Public/Chimera/ChimeraLibrary.h
@@ -0,0 +1,115 @@
+// Copyright Epic Games, Inc. All Rights Reserved.
+
+#pragma once
+
+#include "Kismet/BlueprintFunctionLibrary.h"
+#include "ChimeraLibrary.generated.h"
+
+// (CH)aracter (I)nteraction (M)atching (ERA) =):->
+// subsystem in charge of coordinating Characters availability for motion matched interactions,
+// scheduling motion matching searches and synchronizing animations playback
+
+struct FAnimNode_PoseSearchHistoryCollector_Base;
+class UPoseSearchDatabase;
+
+// input for ChimeraQuery_Pure: it declares that the associated character (AnimInstance) is willing to partecipate in an interction 
+// described by a UMultiAnimAsset (derived by UChimeraAsset) contained in the UPoseSearchDatabase Database
+// with one of the roles in RolesFilter (if empty ANY of the Database roles can be taken)
+// the ChimeraQuery_Pure will ultimately setup a motion matching query using looking for the pose history "PoseHistoryName" 
+// to gather bone and trajectory positions for this character
+// for an interaction to be valid, the query needs to find all the other interacting characters within BroadPhaseRadius,
+// and reach a maximum cost of MaxCost
+USTRUCT(BlueprintType, Category = "Animation|Chimera")
+struct CHIMERA_API FChimeraAvailability
+{
+	GENERATED_BODY()
+
+	// Database describing the interaction. It'll contains multi character UMultiAnimAsset and a schema with multiple skeletons with associated roles
+	UPROPERTY(EditAnywhere, BlueprintReadWrite, Category=Settings)
+	TObjectPtr<UPoseSearchDatabase> Database;
+
+	// roles the character is willing to take to partecipate in this interaction. If empty ANY of the Database roles can be taken
+	UPROPERTY(EditAnywhere, BlueprintReadWrite, Category=Settings)
+	TArray<FName> RolesFilter;
+
+	// the associated character to this FChimeraAvailability will partecipate in an interaction only if all the necessary roles gest assigned to character within BroadPhaseRadius centimeters
+	UPROPERTY(EditAnywhere, BlueprintReadWrite, Category=Settings)
+	float BroadPhaseRadius = 500.f;
+
+	// the associated character to this FChimeraAvailability will not partecipate in an interaction if the motion matching search cost result is higher than MaxCost
+	UPROPERTY(EditAnywhere, BlueprintReadWrite, Category=Settings)
+	float MaxCost = 1000.f;
+};
+
+USTRUCT(BlueprintType, Category="Animation|Chimera")
+struct CHIMERA_API FChimeraBlueprintResult
+{
+	GENERATED_BODY()
+public:
+
+	// animation assigned to this character to partecipate in the interaction
+	UPROPERTY(Transient, VisibleAnywhere, BlueprintReadOnly, Category=State)
+	TObjectPtr<const UObject> SelectedAnimation = nullptr;
+	
+	// SelectedAnimation associated time
+	UPROPERTY(Transient, VisibleAnywhere, BlueprintReadOnly, Category=State)
+	float SelectedTime = 0.f;
+	
+	// SelectedAnimation at SelectedTime is from the continuing pose search
+	UPROPERTY(Transient, VisibleAnywhere, BlueprintReadOnly, Category=State)
+	bool bIsContinuingPoseSearch = false;
+
+	// SelectedAnimation associated play rate
+	UPROPERTY(Transient, VisibleAnywhere, BlueprintReadOnly, Category=State)
+	float WantedPlayRate = 0.f;
+
+	// SelectedAnimation associated looping state
+	UPROPERTY(Transient, VisibleAnywhere, BlueprintReadOnly, Category=State)
+	bool bLoop = false;
+	
+	// SelectedAnimation associated mirror state
+	UPROPERTY(Transient, VisibleAnywhere, BlueprintReadOnly, Category=State)
+	bool bIsMirrored = false;
+	
+	// SelectedAnimation associated BlendParameters (if SelectedAnimation is a UBlendSpace)
+	UPROPERTY(Transient, VisibleAnywhere, BlueprintReadOnly, Category=State)
+	FVector BlendParameters = FVector::ZeroVector;
+
+	// selected SelectedDatabase for this character interaction
+	UPROPERTY(Transient, VisibleAnywhere, BlueprintReadOnly, Category=State)
+	TWeakObjectPtr<const UPoseSearchDatabase> SelectedDatabase = nullptr;
+
+	// associated motion matching search cost for this result
+	UPROPERTY(Transient, VisibleAnywhere, BlueprintReadOnly, Category=State)
+	float SearchCost = MAX_flt;
+
+	// assigned role to this character (AnimInstance)
+	UPROPERTY(Transient, VisibleAnywhere, BlueprintReadOnly, Category=State)
+	FName Role;
+
+	// root bone transform for the character at full aligment
+	UPROPERTY(Transient, VisibleAnywhere, BlueprintReadOnly, Category=State)
+	FTransform FullAlignedActorRootBoneTransform = FTransform::Identity;	
+};
+
+UCLASS()
+class UChimeraLibrary : public UBlueprintFunctionLibrary
+{
+	GENERATED_BODY()
+
+public:
+
+	// function publishing this character (via its AnimInstance) FChimeraAvailability to the UChimeraSubsystem,
+	// FChimeraAvailability represents the character availability to partecipate in an interaction with other characters for the next frame.
+	// that means there will always be one frame delay between publiching availabilities and getting a result back from ChimeraQuery_Pure!
+	// 
+	// if FChimeraBlueprintResult has a valid SelectedAnimation, this will be the animation assigned to this character to partecipate in this interaction.
+	// additional interaction properties, like assigned role, SelectedAnimation time, SearchCost, etc can be found within the result
+	// PoseHistoryName is the name of the pose history node used for the associated motion matching search
+	// if bValidateResultAgainstAvailabilities is true, the result will be invalidated if doesn't respect the new availabilities
+	UFUNCTION(BlueprintPure, Category = "Animation|Chimera", meta = (BlueprintThreadSafe))
+	static FChimeraBlueprintResult ChimeraQuery_Pure(TArray<FChimeraAvailability> Availabilities, UObject* AnimInstance, FName PoseHistoryName, bool bValidateResultAgainstAvailabilities);
+
+	static FChimeraBlueprintResult ChimeraQuery(const TArrayView<const FChimeraAvailability> Availabilities, UObject* AnimInstance, const FAnimNode_PoseSearchHistoryCollector_Base* HistoryCollector, bool bValidateResultAgainstAvailabilities);
+
+};
\ No newline at end of file
diff --git a/Engine/Plugins/Experimental/Animation/Chimera/Source/Runtime/Public/Chimera/ChimeraSubsystem.h b/Engine/Plugins/Experimental/Animation/Chimera/Source/Runtime/Public/Chimera/ChimeraSubsystem.h
new file mode 100644
--- /dev/null
+++ b/Engine/Plugins/Experimental/Animation/Chimera/Source/Runtime/Public/Chimera/ChimeraSubsystem.h
@@ -0,0 +1,112 @@
+// Copyright Epic Games, Inc. All Rights Reserved.
+
+#pragma once
+
+#include "Subsystems/WorldSubsystem.h"
+#include "Chimera/ChimeraLibrary.h"
+#include "ChimeraSubsystem.generated.h"
+
+class UAnimInstance;
+class UChimeraIslandComponent;
+class UPoseSearchDatabase;
+
+namespace UE::Chimera
+{
+	struct FChimeraAvailabilityEx : public FChimeraAvailability
+	{
+		FChimeraAvailabilityEx(const FChimeraAvailability& InAvailability, FName InPoseHistoryName, const FAnimNode_PoseSearchHistoryCollector_Base* InHistoryCollector)
+			: FChimeraAvailability(InAvailability)
+			, PoseHistoryName(InPoseHistoryName)
+			, HistoryCollector(InHistoryCollector)
+		{
+		}
+
+		FString GetPoseHistoryName() const;
+		const FAnimNode_PoseSearchHistoryCollector_Base* GetHistoryCollector(const UAnimInstance* AnimInstance) const;
+
+	private:
+		FName PoseHistoryName;
+		const FAnimNode_PoseSearchHistoryCollector_Base* HistoryCollector = nullptr;
+	};
+	typedef TMap<TWeakObjectPtr<UObject>, TArray<FChimeraAvailabilityEx>> FAvailabilityRequestsMap;
+	typedef TArray<TObjectPtr<UChimeraIslandComponent>> FIslands;
+
+} // namespace UE::Chimera
+
+// World subsystem accepting the publication of characters (via their AnimInstance(s)) FChimeraAvailability, representing the characters willingness to partecipate in an 
+// interaction with other characters from the next frame forward via Query_AnyThread method.
+//
+// The same method will return the FChimeraBlueprintResult from the PREVIOUS Tick processing (categorization of FChimeraAvailability(s) in multiple UChimeraIslandComponent(s)),
+// to the requesting character, containing the animation to play at what time, and the assigned role to partecipate in the selected interaction within the assigned 
+// UChimeraIslandComponent
+
+// Execution model and threading details:
+/////////////////////////////////////////
+// 
+// - by calling UChimeraLibrary::ChimeraQuery_Pure(TArray<FChimeraAvailability> Availabilities, UObject* AnimInstance), characters publish their availabilities
+//   to partecipate in interactions to the UChimeraSubsystem
+// - UChimeraSubsystem::Tick processes those FChimeraAvailability(s) and creates/updates UChimeraIslandComponent. For each UChimeraIslandComponent it injects 
+//   a tick prerequisite via UChimeraIslandComponent::InjectToActor (that calls AddTickPrerequisiteComponent) to all the Actors in the same island.
+//   NoTe: the next frame the execution will be:
+//			for each island[k]
+//			{
+//				for each Actor[k][i]
+//				{
+//					Tick CharacterMovementComponent[k][i]
+//				}
+// 
+//				Tick ChimeraIslandComponent[k]
+//
+//				for each Actor[k][i]
+//				{
+//					Tick SkeletalMeshComponent[k][i]
+//				}
+//			}
+// - next frame UChimeraLibrary::ChimeraQuery_Pure(TArray<FChimeraAvailability> Availabilities, UObject* AnimInstance), with the context of all the published 
+//   availabilities and created islands, will find the associated UChimeraIslandComponent to the AnimInstance and call UChimeraIslandComponent::DoSearch_AnyThread 
+//   (via UChimeraSubsystem::Query_AnyThread) that will perform ALL (YES, ALL, so the bigger the island the slower the execution) the motion matching searches
+//   for all the possible Actors / databases / Roles permutations, and populate UChimeraIslandComponent::SearchResults with ALL the results for the island.
+//   Ultimately the ChimeraQuery_Pure will return the SearchResults associated to the requesting AnimInstance with information about what animation to play
+//   at what time with wich Role.
+
+UCLASS(Experimental, Category="Animation|Chimera")
+class CHIMERA_API UChimeraSubsystem : public UTickableWorldSubsystem
+{
+	GENERATED_BODY()
+
+public:
+	static UChimeraSubsystem* GetSubsystem_AnyThread(UObject* AnimInstance);
+
+	virtual void Deinitialize() override;
+
+	// it processes FChimeraAvailability(s) and creates/updates UChimeraIslandComponent
+	virtual void Tick(float DeltaSeconds) override;
+	
+	virtual TStatId GetStatId() const override;
+
+	// publishing FChimeraAvailability(s) for the requesting character (AnimInstance) and getting the FChimeraBlueprintResult from the PREVIOUS Tick update
+	// containing the animation to play at what time, and the assigned role to partecipate in the selected interaction
+	// Either a PoseHistoryName or a HistoryCollector are required to perform the associated motion matching searches
+	void Query_AnyThread(const TArrayView<const FChimeraAvailability> Availabilities, UObject* AnimInstance, FChimeraBlueprintResult& Result,
+		FName PoseHistoryName, const FAnimNode_PoseSearchHistoryCollector_Base* HistoryCollector, bool bValidateResultAgainstAvailabilities);
+private:
+	UChimeraIslandComponent& CreateIsland();
+	UChimeraIslandComponent& GetAvailableIsland();
+
+	void DestroyIsland(int32 Index);
+	void DestroyAllIslands();
+	void UninjectAllIslands();
+	bool ValidateAllIslands() const;
+	UChimeraIslandComponent* FindIsland(UObject* InAnimInstance);
+	
+	void DebugDraw() const;
+
+	UE::Chimera::FAvailabilityRequestsMap AvailabilityRequestsMap;
+	FCriticalSection AvailabilityRequestsMapMutex;
+
+	// array of groups of characters that needs to be anaylzed together for possible interactions
+	UE::Chimera::FIslands Islands;
+
+	// critical section to retrieve the subsystem in a thread safe manner
+	static FCriticalSection RetrieveSubsystemMutex;
+};
\ No newline at end of file
