chimera - fix deadlock occurring between UPoseSearchLibrary::MotionMatch and UAnimInstance::HandleExistingParallelEvaluationTask while locking UChimeraIslandComponent::SearchResultsMutex

diff --git a/Engine/Plugins/Experimental/Animation/Chimera/Source/Runtime/Private/AnimNode_Chimera.cpp b/Engine/Plugins/Experimental/Animation/Chimera/Source/Runtime/Private/AnimNode_Chimera.cpp
--- a/Engine/Plugins/Experimental/Animation/Chimera/Source/Runtime/Private/AnimNode_Chimera.cpp
+++ b/Engine/Plugins/Experimental/Animation/Chimera/Source/Runtime/Private/AnimNode_Chimera.cpp
@@ -118,7 +118,7 @@ void FAnimNode_Chimera::UpdateAssetPlayer(const FAnimationUpdateContext& Context
 		}
 	}
 	
-	// calculating the translation and rotation warp lerps, used to wapr the root transform towards the last computed FullAlignedActorRootBoneTransform
+	// calculating the translation and rotation warp lerps, used to warp the root transform towards the last computed FullAlignedActorRootBoneTransform
 	const float Sign = bInteracting ? 1.f : -1.f;
 	if (BlendTime > UE_KINDA_SMALL_NUMBER)
 	{
diff --git a/Engine/Plugins/Experimental/Animation/Chimera/Source/Runtime/Private/ChimeraIslandComponent.cpp b/Engine/Plugins/Experimental/Animation/Chimera/Source/Runtime/Private/ChimeraIslandComponent.cpp
--- a/Engine/Plugins/Experimental/Animation/Chimera/Source/Runtime/Private/ChimeraIslandComponent.cpp
+++ b/Engine/Plugins/Experimental/Animation/Chimera/Source/Runtime/Private/ChimeraIslandComponent.cpp
@@ -218,9 +218,12 @@ void UChimeraIslandComponent::TickComponent(float DeltaTime, enum ELevelTick Tic
 
 void UChimeraIslandComponent::DebugDraw(const FColor& Color) const
 {
+	// called only by UChimeraSubsystem::Tick so no need to lock SearchResultsMutex to protect the read of SearchContexts
 #if ENABLE_DRAW_DEBUG
 	using namespace UE::Chimera;
 
+	check(IsInGameThread());
+
 	if (CVarChimeraShowIslands.GetValueOnAnyThread())
 	{
 		for (const FSearchContext& SearchContext : SearchContexts)
@@ -241,6 +244,8 @@ void UChimeraIslandComponent::DebugDraw(const FColor& Color) const
 
 void UChimeraIslandComponent::InjectToActor(AActor* Actor)
 {
+	check(IsInGameThread());
+
 	// Called by UChimeraSubsystem::Tick when there aren't animation jobs flying. No need to FScopeLock Lock(&Mutex);
 	if (Actor)
 	{
@@ -263,6 +268,7 @@ void UChimeraIslandComponent::AddSearchContext(const UE::Chimera::FSearchContext
 {
 #if DO_CHECK
 	check(SearchContext.IsValid());
+	check(IsInGameThread());
 
 	for (const UE::Chimera::FSearchContext& ContainedSearchContext : SearchContexts)
 	{
@@ -274,11 +280,13 @@ void UChimeraIslandComponent::AddSearchContext(const UE::Chimera::FSearchContext
 
 void UChimeraIslandComponent::ResetSearchContexts()
 {
+	check(IsInGameThread());
 	SearchContexts.Reset();
 }
 
 void UChimeraIslandComponent::ResetSearchResults()
 {
+	check(IsInGameThread());
 	SearchResults.Reset();
 	bSearchPerfomed = false;
 }
@@ -302,42 +310,40 @@ void UChimeraIslandComponent::UninjectFromAllActors()
 
 bool UChimeraIslandComponent::IsUninjected()
 {
-	const bool bIsUninjected = SkeletalMeshComponents.IsEmpty();
-
-#if DO_CHECK
-	if (bIsUninjected)
-	{
-		// making sure properties are consistent within each other
-		check(CharacterMovementComponents.IsEmpty());
-		check(SearchContexts.IsEmpty());
-
-		FScopeLock Lock(&SearchResultsMutex);
-		check(SearchResults.IsEmpty());
-		check(bSearchPerfomed == false);
-	}
-#endif // DO_CHECK
-
-	return bIsUninjected;
+	return SkeletalMeshComponents.IsEmpty();
 }
 
 bool UChimeraIslandComponent::DoSearch_AnyThread(UObject* AnimInstance, FChimeraBlueprintResult& Result)
 {
 	using namespace UE::Chimera;
-	
-	FScopeLock Lock(&SearchResultsMutex);
 
-	if (!bSearchPerfomed)
+	bool bDoPerfomSearch = false;
+	{
+		// thread safety note!
+		// goal:	avoiding deadlock between SearchResultsMutex lock and waiting for UAnimInstance::HandleExistingParallelEvaluationTask.
+		// why:		UPoseSearchLibrary::MotionMatch could call via AnimInstance GetProxyOnAnyThread<FAnimInstanceProxy>() that, if on GameThread, 
+		//			could call UAnimInstance::HandleExistingParallelEvaluationTask 
+		// fix:		avoid UPoseSearchLibrary::MotionMatch calls wrapped by any lock, at the cost of eventually (by design should be NEVER) performing the searches twice
+		//			By design we should inject ticks dependencies (added by UChimeraIslandComponent::InjectToActor via AddTickPrerequisiteComponent), so concurrently 
+		//			fly of UChimeraIslandComponent within the same island that requires searches is forbidden
+		FScopeLock Lock(&SearchResultsMutex);
+		bDoPerfomSearch = !bSearchPerfomed;
+	}
+
+	if (bDoPerfomSearch)
 	{
 		FMemMark Mark(FMemStack::Get());
 
 		TArray<UAnimInstance*, TInlineAllocator<UE::PoseSearch::PreallocatedRolesNum, TMemStackAllocator<>>> AnimInstances;
 		TArray<const UE::PoseSearch::IPoseHistory*, TInlineAllocator<UE::PoseSearch::PreallocatedRolesNum, TMemStackAllocator<>>> PoseHistories;
 		TArray<UE::PoseSearch::FSearchResult, TMemStackAllocator<>> PoseSearchResults;
+
+		// SearchContexts are modified only by UChimeraSubsystem::Tick and constant otherwise, so it's safe to access them in a threaded enviroment without locks
 		PoseSearchResults.SetNum(SearchContexts.Num());
 
 		for (int32 SearchIndex = 0; SearchIndex < SearchContexts.Num(); ++SearchIndex)
 		{
-			FSearchContext& SearchContext = SearchContexts[SearchIndex];
+			const FSearchContext& SearchContext = SearchContexts[SearchIndex];
 			const UPoseSearchDatabase* Database = SearchContext.Database.Get();
 			if (!Database)
 			{
@@ -352,7 +358,7 @@ bool UChimeraIslandComponent::DoSearch_AnyThread(UObject* AnimInstance, FChimera
 			}
 
 			AnimInstances.Reset();
-			for (TWeakObjectPtr<UAnimInstance>& AnimInstancePtr : SearchContext.AnimInstances)
+			for (const TWeakObjectPtr<UAnimInstance>& AnimInstancePtr : SearchContext.AnimInstances)
 			{
 				UAnimInstance* SearchContextAnimInstance = AnimInstancePtr.Get();
 				if (!SearchContextAnimInstance)
@@ -372,6 +378,8 @@ bool UChimeraIslandComponent::DoSearch_AnyThread(UObject* AnimInstance, FChimera
 
 			const UObject* AssetsToSearch[] = { Database };
 			FPoseSearchFutureProperties PoseSearchFutureProperties;
+
+			// @todo: we could perform multiple UPoseSearchLibrary::MotionMatch in parallel!
 			const UE::PoseSearch::FSearchResult PoseSearchResult = UPoseSearchLibrary::MotionMatch(AnimInstances, SearchContext.Roles,
 				PoseHistories, AssetsToSearch, SearchContext.ContinuingProperties, PoseSearchFutureProperties);
 
@@ -384,28 +392,47 @@ bool UChimeraIslandComponent::DoSearch_AnyThread(UObject* AnimInstance, FChimera
 		////////////////////////////////////////////////////////////////////////////////////////////////////////////////
 		// WIP!
 		// @todo: figure out multiple policies to sue the most characters? right now only the best search is "valid"
-		FPoseSearchCost BestPoseCost;
-		int32 BestSearchIndex = INDEX_NONE;
-		for (int32 SearchIndex = 0; SearchIndex < PoseSearchResults.Num(); ++SearchIndex)
+		if (!PoseSearchResults.IsEmpty())
 		{
-			const UE::PoseSearch::FSearchResult& PoseSearchResult = PoseSearchResults[SearchIndex];
-			if (PoseSearchResult.IsValid() && PoseSearchResult.PoseCost < BestPoseCost)
+			// locking to update SearchResults and bSearchPerfomed
+			FScopeLock Lock(&SearchResultsMutex);
+
+			FPoseSearchCost BestPoseCost;
+			int32 BestSearchIndex = INDEX_NONE;
+			for (int32 SearchIndex = 0; SearchIndex < PoseSearchResults.Num(); ++SearchIndex)
 			{
-				BestPoseCost = PoseSearchResult.PoseCost;
-				BestSearchIndex = SearchIndex;
+				const UE::PoseSearch::FSearchResult& PoseSearchResult = PoseSearchResults[SearchIndex];
+				if (PoseSearchResult.IsValid() && PoseSearchResult.PoseCost < BestPoseCost)
+				{
+					BestPoseCost = PoseSearchResult.PoseCost;
+					BestSearchIndex = SearchIndex;
+				}
 			}
-		}
 
-		if (BestSearchIndex != INDEX_NONE)
-		{
-			SearchResults.SetNum(1);
-			SearchResults[0] = InitSearchResult(PoseSearchResults[BestSearchIndex], BestSearchIndex, SearchContexts[BestSearchIndex]);
-		}
-		////////////////////////////////////////////////////////////////////////////////////////////////////////////////
+			if (BestSearchIndex != INDEX_NONE)
+			{
+				SearchResults.SetNum(1);
+				SearchResults[0] = InitSearchResult(PoseSearchResults[BestSearchIndex], BestSearchIndex, SearchContexts[BestSearchIndex]);
+			}
+			////////////////////////////////////////////////////////////////////////////////////////////////////////////////
 
-		bSearchPerfomed = true;
+			bSearchPerfomed = true;
+	
+			// calling this funtion within this scope since we already locked SearchResultsMutex
+			return GetResult_AnyThread(AnimInstance, Result);
+		}
 	}
 
+	return GetResult_AnyThread(AnimInstance, Result);
+}
+
+bool UChimeraIslandComponent::GetResult_AnyThread(UObject* AnimInstance, FChimeraBlueprintResult& Result)
+{
+	using namespace UE::Chimera;
+
+	// locking to read SearchResults
+	FScopeLock Lock(&SearchResultsMutex);
+
 	// looking for AnimInstance in SearchResults to fill up Result
 	for (const FSearchResult& SearchResult : SearchResults)
 	{
@@ -469,6 +496,9 @@ bool UChimeraIslandComponent::DoSearch_AnyThread(UObject* AnimInstance, FChimera
 
 const UE::Chimera::FSearchResult* UChimeraIslandComponent::FindSearchResult(const UE::Chimera::FSearchContext& SearchContext) const
 {
+	// called only by UChimeraSubsystem::Tick via UChimeraSubsystem::PopulateContinuingProperties so no need to lock SearchResultsMutex to protect the read of SearchResults
+	check(IsInGameThread());
+
 	using namespace UE::Chimera;
 
 	// searching for InSearchContext in all the SearchContexts referenced by valid active SearchResults
diff --git a/Engine/Plugins/Experimental/Animation/Chimera/Source/Runtime/Private/ChimeraSubsystem.cpp b/Engine/Plugins/Experimental/Animation/Chimera/Source/Runtime/Private/ChimeraSubsystem.cpp
--- a/Engine/Plugins/Experimental/Animation/Chimera/Source/Runtime/Private/ChimeraSubsystem.cpp
+++ b/Engine/Plugins/Experimental/Animation/Chimera/Source/Runtime/Private/ChimeraSubsystem.cpp
@@ -204,30 +204,6 @@ namespace UE::Chimera
 		}
 	}
 
-	static void PopulateContinuingProperties(FSearchContext& SearchContext, const FIslands& Islands, float DeltaSeconds)
-	{
-		// searching this SearchContext in all the islands to initialize its continuing pose
-		for (const UChimeraIslandComponent* Island : Islands)
-		{
-			if (const FSearchResult* SearchResult = Island->FindSearchResult(SearchContext))
-			{
-				// is still valid...
-				if (SearchResult->IsValid())
-				{
-					if (const UE::PoseSearch::FSearchIndexAsset* SearchIndexAsset = SearchResult->GetSearchIndexAsset())
-					{
-						if (const FPoseSearchDatabaseAnimationAssetBase* DatabaseAsset = SearchResult->Database->GetDatabaseAnimationAsset<FPoseSearchDatabaseAnimationAssetBase>(*SearchIndexAsset))
-						{
-							SearchContext.ContinuingProperties.PlayingAsset = DatabaseAsset->GetAnimationAsset();
-							SearchContext.ContinuingProperties.PlayingAssetAccumulatedTime = SearchResult->AssetTime + DeltaSeconds;
-						}
-					}
-				}
-				break;
-			}
-		}
-	}
-
 	template <typename DataType, typename EvaluateCombinationType>
 	static void GeneratePermutationsRecursive(const TArrayView<DataType> Data, int32 DataIndex, TArrayView<int32> Combination, int32 CombinationIndex, EvaluateCombinationType EvaluateCombination)
 	{
@@ -299,6 +275,8 @@ void UChimeraSubsystem::DestroyAllIslands()
 
 void UChimeraSubsystem::UninjectAllIslands()
 {
+	check(IsInGameThread());
+
 	for (TObjectPtr<UChimeraIslandComponent>& IslandPtr : Islands)
 	{
 		UChimeraIslandComponent& Island = *IslandPtr.Get();
@@ -340,6 +318,32 @@ bool UChimeraSubsystem::ValidateAllIslands() const
 	return true;
 }
 
+void UChimeraSubsystem::PopulateContinuingProperties(UE::Chimera::FSearchContext& SearchContext, float DeltaSeconds) const
+{
+	check(IsInGameThread());
+
+	// searching this SearchContext in all the islands to initialize its continuing pose
+	for (const UChimeraIslandComponent* Island : Islands)
+	{
+		if (const UE::Chimera::FSearchResult* SearchResult = Island->FindSearchResult(SearchContext))
+		{
+			// is still valid...
+			if (SearchResult->IsValid())
+			{
+				if (const UE::PoseSearch::FSearchIndexAsset* SearchIndexAsset = SearchResult->GetSearchIndexAsset())
+				{
+					if (const FPoseSearchDatabaseAnimationAssetBase* DatabaseAsset = SearchResult->Database->GetDatabaseAnimationAsset<FPoseSearchDatabaseAnimationAssetBase>(*SearchIndexAsset))
+					{
+						SearchContext.ContinuingProperties.PlayingAsset = DatabaseAsset->GetAnimationAsset();
+						SearchContext.ContinuingProperties.PlayingAssetAccumulatedTime = SearchResult->AssetTime + DeltaSeconds;
+					}
+				}
+			}
+			break;
+		}
+	}
+}
+
 UChimeraIslandComponent* UChimeraSubsystem::FindIsland(UObject* InAnimInstance)
 {
 	if (InAnimInstance)
@@ -518,7 +522,7 @@ void UChimeraSubsystem::Tick(float DeltaSeconds)
 	// for each valid SearchContexts we try to figure out the continuing pose properties from the current Islands
 	for (FSearchContext& SearchContext : SearchContexts)
 	{
-		PopulateContinuingProperties(SearchContext, Islands, DeltaSeconds);
+		PopulateContinuingProperties(SearchContext, DeltaSeconds);
 	}
 
 	// grouping SearchContexts in different Islands
diff --git a/Engine/Plugins/Experimental/Animation/Chimera/Source/Runtime/Public/Chimera/ChimeraIslandComponent.h b/Engine/Plugins/Experimental/Animation/Chimera/Source/Runtime/Public/Chimera/ChimeraIslandComponent.h
--- a/Engine/Plugins/Experimental/Animation/Chimera/Source/Runtime/Public/Chimera/ChimeraIslandComponent.h
+++ b/Engine/Plugins/Experimental/Animation/Chimera/Source/Runtime/Public/Chimera/ChimeraIslandComponent.h
@@ -48,6 +48,12 @@ UCLASS(Experimental, Category="Animation|Chimera")
 class UChimeraIslandComponent : public UActorComponent
 {
 public:
+	bool DoSearch_AnyThread(UObject* AnimInstance, FChimeraBlueprintResult& Result);
+
+private:
+	// to minimize the usage of lock we expose only thread safe APIs and allow private calls only via the power of friendship
+	friend class UChimeraSubsystem;
+
 	GENERATED_BODY()
 
 	virtual void TickComponent(float DeltaTime, enum ELevelTick TickType, FActorComponentTickFunction *ThisTickFunction) override;
@@ -56,12 +62,12 @@ public:
 	void UninjectFromAllActors();
 	bool IsUninjected();
 
-	bool DoSearch_AnyThread(UObject* AnimInstance, FChimeraBlueprintResult& Result);
-
 	void AddSearchContext(const UE::Chimera::FSearchContext& SearchContext);
 	void ResetSearchContexts();
 	void ResetSearchResults();
 
+	bool GetResult_AnyThread(UObject* AnimInstance, FChimeraBlueprintResult& Result);
+
 	void DebugDraw(const FColor& Color = FColor::Red) const;
 
 	const TArray<TWeakObjectPtr<UCharacterMovementComponent>>& GetCharacterMovementComponents() const { return CharacterMovementComponents; }
@@ -71,7 +77,6 @@ public:
 
 	const UE::Chimera::FSearchResult* FindSearchResult(const UE::Chimera::FSearchContext& SearchContext) const;
 
-private:
 	TArray<TWeakObjectPtr<UCharacterMovementComponent>> CharacterMovementComponents;
 	TArray<TWeakObjectPtr<USkeletalMeshComponent>> SkeletalMeshComponents;
 	
diff --git a/Engine/Plugins/Experimental/Animation/Chimera/Source/Runtime/Public/Chimera/ChimeraSubsystem.h b/Engine/Plugins/Experimental/Animation/Chimera/Source/Runtime/Public/Chimera/ChimeraSubsystem.h
--- a/Engine/Plugins/Experimental/Animation/Chimera/Source/Runtime/Public/Chimera/ChimeraSubsystem.h
+++ b/Engine/Plugins/Experimental/Animation/Chimera/Source/Runtime/Public/Chimera/ChimeraSubsystem.h
@@ -12,6 +12,8 @@ class UPoseSearchDatabase;
 
 namespace UE::Chimera
 {
+	struct FSearchContext;
+
 	struct FChimeraAvailabilityEx : public FChimeraAvailability
 	{
 		FChimeraAvailabilityEx(const FChimeraAvailability& InAvailability, FName InPoseHistoryName, const FAnimNode_PoseSearchHistoryCollector_Base* InHistoryCollector)
@@ -97,6 +99,8 @@ private:
 	void DestroyAllIslands();
 	void UninjectAllIslands();
 	bool ValidateAllIslands() const;
+
+	void PopulateContinuingProperties(UE::Chimera::FSearchContext& SearchContext, float DeltaSeconds) const;
 	UChimeraIslandComponent* FindIsland(UObject* InAnimInstance);
 	
 	void DebugDraw() const;
