chimera - parented FAnimNode_Chimera to FAnimNode_BlendStack, added evaluation mode, warp improvements

diff --git a/Engine/Plugins/Animation/BlendStack/Source/Runtime/Private/AnimNode_BlendStack.cpp b/Engine/Plugins/Animation/BlendStack/Source/Runtime/Private/AnimNode_BlendStack.cpp
--- a/Engine/Plugins/Animation/BlendStack/Source/Runtime/Private/AnimNode_BlendStack.cpp
+++ b/Engine/Plugins/Animation/BlendStack/Source/Runtime/Private/AnimNode_BlendStack.cpp
@@ -1131,23 +1131,17 @@ void FAnimNode_BlendStack_Standalone::BlendWithPose(FAnimationPoseData& InOutPos
 /////////////////////////////////////////////////////
 // FAnimNode_BlendStack
 
-void FAnimNode_BlendStack::UpdateAssetPlayer(const FAnimationUpdateContext& Context)
+bool FAnimNode_BlendStack::NeedsReset(const FAnimationUpdateContext& Context) const
 {
 	const bool bNeedsReset =
 		bResetOnBecomingRelevant &&
 		UpdateCounter.HasEverBeenUpdated() &&
 		!UpdateCounter.WasSynchronizedCounter(Context.AnimInstanceProxy->GetUpdateCounter());
+	return bNeedsReset;
+}
 
-	if (bNeedsReset)
-	{
-		Reset();
-		bForceBlendNextUpdate = false;
-	}
-
-	UpdateCounter.SynchronizeWith(Context.AnimInstanceProxy->GetUpdateCounter());
-
-	GetEvaluateGraphExposedInputs().Execute(Context);
-
+bool FAnimNode_BlendStack::ConditionalBlendTo(const FAnimationUpdateContext& Context)
+{
 	bool bExecuteBlendTo = false;
 	if (AnimationAsset == nullptr && !bForceBlendNextUpdate)
 	{
@@ -1192,6 +1186,27 @@ void FAnimNode_BlendStack::UpdateAssetPlayer(const FAnimationUpdateContext& Cont
 			GetGroupName(), GetGroupRole(), GetGroupMethod());
 	}
 
+	return bExecuteBlendTo;
+}
+
+void FAnimNode_BlendStack::Reset()
+{
+	Super::Reset();
+	bForceBlendNextUpdate = false;
+}
+
+void FAnimNode_BlendStack::UpdateAssetPlayer(const FAnimationUpdateContext& Context)
+{
+	if (NeedsReset(Context))
+	{
+		Reset();
+	}
+
+	UpdateCounter.SynchronizeWith(Context.AnimInstanceProxy->GetUpdateCounter());
+
+	GetEvaluateGraphExposedInputs().Execute(Context);
+
+	const bool bExecuteBlendTo = ConditionalBlendTo(Context);
 	const bool bDidBlendToRequestAnInertialBlend = bExecuteBlendTo && bUseInertialBlend;
 	UE::Anim::TOptionalScopedGraphMessage<UE::Anim::FAnimInertializationSyncScope> InertializationSync(bDidBlendToRequestAnInertialBlend, Context);
 	
diff --git a/Engine/Plugins/Animation/BlendStack/Source/Runtime/Public/BlendStack/AnimNode_BlendStack.h b/Engine/Plugins/Animation/BlendStack/Source/Runtime/Public/BlendStack/AnimNode_BlendStack.h
--- a/Engine/Plugins/Animation/BlendStack/Source/Runtime/Public/BlendStack/AnimNode_BlendStack.h
+++ b/Engine/Plugins/Animation/BlendStack/Source/Runtime/Public/BlendStack/AnimNode_BlendStack.h
@@ -159,7 +159,7 @@ struct BLENDSTACK_API FAnimNode_BlendStack_Standalone : public FAnimNode_AssetPl
 		bool bUseInertialBlend = false, const FVector& BlendParameters = FVector::Zero(), float PlayRate = 1.f, float ActivationDelay = 0.f,
 		FName GroupName = NAME_None, EAnimGroupRole::Type GroupRole = EAnimGroupRole::CanBeLeader, EAnimSyncMethod Method = EAnimSyncMethod::DoNotSync, bool bOverridePositionWhenJoiningSyncGroupAsLeader = false);
 	void UpdatePlayRate(float PlayRate);
-	void Reset();
+	virtual void Reset();
 
 	// FAnimNode_AssetPlayerBase interface
 	virtual UAnimationAsset* GetAnimAsset() const;
@@ -290,27 +290,27 @@ struct BLENDSTACK_API FAnimNode_BlendStack : public FAnimNode_BlendStack_Standal
 	UPROPERTY(EditAnywhere, Category = Settings)
 	bool bResetOnBecomingRelevant = true;
 
-	#if WITH_EDITORONLY_DATA
-    	// The group name that we synchronize with (NAME_None if it is not part of any group). Note that
-    	// this is the name of the group used to sync the output of this node - it will not force
-    	// syncing of animations contained by it.
-    	UPROPERTY(EditAnywhere, Category = Sync, meta = (FoldProperty))
-    	FName GroupName = NAME_None;
+#if WITH_EDITORONLY_DATA
+    // The group name that we synchronize with (NAME_None if it is not part of any group). Note that
+    // this is the name of the group used to sync the output of this node - it will not force
+    // syncing of animations contained by it.
+    UPROPERTY(EditAnywhere, Category = Sync, meta = (FoldProperty))
+    FName GroupName = NAME_None;
     
-    	// The role this node can assume within the group (ignored if GroupName is not set). Note
-    	// that this is the role of the output of this node, not of animations contained by it.
-    	UPROPERTY(EditAnywhere, Category = Sync, meta = (FoldProperty))
-    	TEnumAsByte<EAnimGroupRole::Type> GroupRole = EAnimGroupRole::CanBeLeader;
+    // The role this node can assume within the group (ignored if GroupName is not set). Note
+    // that this is the role of the output of this node, not of animations contained by it.
+    UPROPERTY(EditAnywhere, Category = Sync, meta = (FoldProperty))
+    TEnumAsByte<EAnimGroupRole::Type> GroupRole = EAnimGroupRole::CanBeLeader;
     
-    	// How this node will synchronize with other animations. Note that this determines how the output
-    	// of this node is used for synchronization, not of animations contained by it.
-    	UPROPERTY(EditAnywhere, Category = Sync, meta = (FoldProperty))
-    	EAnimSyncMethod Method = EAnimSyncMethod::DoNotSync;
+    // How this node will synchronize with other animations. Note that this determines how the output
+    // of this node is used for synchronization, not of animations contained by it.
+    UPROPERTY(EditAnywhere, Category = Sync, meta = (FoldProperty))
+    EAnimSyncMethod Method = EAnimSyncMethod::DoNotSync;
     
-    	// If true, "Relevant anim" nodes that look for the highest weighted animation in a state will ignore this node
-    	UPROPERTY(EditAnywhere, Category = Relevancy, meta = (FoldProperty, PinHiddenByDefault))
-    	bool bIgnoreForRelevancyTest = false;
-    #endif // WITH_EDITORONLY_DATA
+    // If true, "Relevant anim" nodes that look for the highest weighted animation in a state will ignore this node
+    UPROPERTY(EditAnywhere, Category = Relevancy, meta = (FoldProperty, PinHiddenByDefault))
+    bool bIgnoreForRelevancyTest = false;
+#endif // WITH_EDITORONLY_DATA
 
 	// FAnimNode_Base interface
 	virtual void UpdateAssetPlayer(const FAnimationUpdateContext& Context) override;
@@ -327,8 +327,12 @@ struct BLENDSTACK_API FAnimNode_BlendStack : public FAnimNode_BlendStack_Standal
 
 	// Force a blend on the next update, even if the anim sequence has not changed.
 	void ForceBlendNextUpdate();
+	virtual void Reset() override;
+
+protected:
+	bool NeedsReset(const FAnimationUpdateContext& Context) const;
+	bool ConditionalBlendTo(const FAnimationUpdateContext& Context);
 
-private:
 	// Update Counter for detecting being relevant
 	FGraphTraversalCounter UpdateCounter;
 
diff --git a/Engine/Plugins/Experimental/Animation/Chimera/Source/Runtime/Private/AnimNode_Chimera.cpp b/Engine/Plugins/Experimental/Animation/Chimera/Source/Runtime/Private/AnimNode_Chimera.cpp
--- a/Engine/Plugins/Experimental/Animation/Chimera/Source/Runtime/Private/AnimNode_Chimera.cpp
+++ b/Engine/Plugins/Experimental/Animation/Chimera/Source/Runtime/Private/AnimNode_Chimera.cpp
@@ -39,29 +39,31 @@ void FAnimNode_Chimera::CacheBones_AnyThread(const FAnimationCacheBonesContext&
 	Source.CacheBones(Context);
 }
 
+void FAnimNode_Chimera::Reset()
+{
+	Super::Reset();
+	BlendLerp = 0.f;
+	TranslationWarpLerp = 0.f;
+	RotationWarpLerp = 0.f;
+	bWasInteracting = false;
+}
+
 void FAnimNode_Chimera::UpdateAssetPlayer(const FAnimationUpdateContext& Context)
 {
 	using namespace UE::PoseSearch;
 
-	const bool bNeedsReset =
-		bResetOnBecomingRelevant &&
-		UpdateCounter.HasEverBeenUpdated() &&
-		!UpdateCounter.WasSynchronizedCounter(Context.AnimInstanceProxy->GetUpdateCounter());
+	if (NeedsReset(Context))
+	{
+		Reset();
+	}
 
 	UpdateCounter.SynchronizeWith(Context.AnimInstanceProxy->GetUpdateCounter());
 
 	GetEvaluateGraphExposedInputs().Execute(Context);
-	
-	if (bNeedsReset)
-	{
-		Reset();
-		BlendLerp = 0.f;
-		TranslationWarpLerp = 0.f;
-		RotationWarpLerp = 0.f;
-	}
 
-	bool bInteracting = false;
-	bool bExecuteBlendTo = false;
+	bool bBlendToExecuted = ConditionalBlendTo(Context);
+
+	bool bIsInteracting = false;
 	const float DeltaTime = Context.GetDeltaTime();
 	FPoseHistoryProvider* PoseHistoryProvider = Context.GetMessage<FPoseHistoryProvider>();
 	if (!PoseHistoryProvider)
@@ -77,28 +79,44 @@ void FAnimNode_Chimera::UpdateAssetPlayer(const FAnimationUpdateContext& Context
 		{
 			check(Result.SelectedDatabase != nullptr);
 
-			UAnimationAsset* AnimationAsset = MultiAnimAsset->GetAnimationAsset(Result.Role);
-			check(AnimationAsset);
+			UAnimationAsset* RoledAnimationAsset = MultiAnimAsset->GetAnimationAsset(Result.Role);
+			check(RoledAnimationAsset);
 
-			FullAlignedActorRootBoneTransform = Result.FullAlignedActorRootBoneTransform;
-			bInteracting = true;
+			bIsInteracting = true;
 
-			if (AnimPlayers.IsEmpty())
+			bool bExecuteBlendTo = false;
+			bool bUpdatePropertiesFromResult = false;
+			if (!bWasInteracting || AnimPlayers.IsEmpty())
 			{
 				bExecuteBlendTo = true;
+				bUpdatePropertiesFromResult = true;
 			}
-			else
+			else if (EvaluationMode == EChimeraEvaluationMode::ContinuousReselection)
 			{
 				const FBlendStackAnimPlayer& MainAnimPlayer = AnimPlayers[0];
 				const UAnimationAsset* PlayingAnimationAsset = MainAnimPlayer.GetAnimationAsset();
 
-				if (AnimationAsset != PlayingAnimationAsset ||
+				if (RoledAnimationAsset != PlayingAnimationAsset ||
 					Result.bIsMirrored != MainAnimPlayer.GetMirror() ||
 					Result.BlendParameters != MainAnimPlayer.GetBlendParameters() ||
 					!Result.bIsContinuingPoseSearch)
 				{
 					bExecuteBlendTo = true;
 				}
+
+				bUpdatePropertiesFromResult = true;
+			}
+			else if (RoledAnimationAsset == AnimPlayers[0].GetAnimationAsset() && Result.bIsContinuingPoseSearch)
+			{
+				// we don't update FullAlignedActorRootBoneTransform since we're not planning to blend into the newly selected animation here
+				bUpdatePropertiesFromResult = true;
+			}
+
+			if (bUpdatePropertiesFromResult)
+			{
+				FullAlignedActorRootBoneTransform = Result.FullAlignedActorRootBoneTransform;
+				WantedPlayRate = Result.WantedPlayRate;
+				BlendParameters = Result.BlendParameters;
 			}
 
 			if (bExecuteBlendTo)
@@ -106,20 +124,23 @@ void FAnimNode_Chimera::UpdateAssetPlayer(const FAnimationUpdateContext& Context
 				const FPoseSearchRoledSkeleton* RoledSkeleton = Result.SelectedDatabase->Schema->GetRoledSkeleton(Result.Role);
 				check(RoledSkeleton);
 
-				BlendTo(Context, AnimationAsset, Result.SelectedTime, Result.bLoop, Result.bIsMirrored, RoledSkeleton->MirrorDataTable.Get(),
-					BlendTime, BlendProfile, BlendOption, bUseInertialBlend, Result.BlendParameters, Result.WantedPlayRate);
-			}
+				BlendTo(Context, RoledAnimationAsset, Result.SelectedTime, Result.bLoop, Result.bIsMirrored, RoledSkeleton->MirrorDataTable.Get(),
+					BlendTime, BlendProfile, BlendOption, bUseInertialBlend, BlendParameters, WantedPlayRate);
 
-			// @todo: do we still need FAnimInertializationSyncScope?
-			//const bool bDidBlendToRequestAnInertialBlend = bExecuteBlendTo && bUseInertialBlend;
-			//UE::Anim::TOptionalScopedGraphMessage<UE::Anim::FAnimInertializationSyncScope> InertializationSync(bDidBlendToRequestAnInertialBlend, Context);
+				bBlendToExecuted = true;
+			}
 
-			UpdatePlayRate(Result.WantedPlayRate);
 		}
 	}
+
+	const bool bDidBlendToRequestAnInertialBlend = bBlendToExecuted && bUseInertialBlend;
+	UE::Anim::TOptionalScopedGraphMessage<UE::Anim::FAnimInertializationSyncScope> InertializationSync(bDidBlendToRequestAnInertialBlend, Context);
 	
+	UpdatePlayRate(WantedPlayRate);
+	UpdateBlendspaceParameters(BlendspaceUpdateMode, BlendParameters);
+
 	// calculating the translation and rotation warp lerps, used to warp the root transform towards the last computed FullAlignedActorRootBoneTransform
-	const float Sign = bInteracting ? 1.f : -1.f;
+	const float Sign = bIsInteracting ? 1.f : -1.f;
 	if (BlendTime > UE_KINDA_SMALL_NUMBER)
 	{
 		BlendLerp = FMath::Clamp(BlendLerp + (Sign * DeltaTime / BlendTime), 0.f, 1.f);
@@ -142,25 +163,40 @@ void FAnimNode_Chimera::UpdateAssetPlayer(const FAnimationUpdateContext& Context
 	}
 	Source.Update(SourceContext);
 	
-	// updating blend stack
-	if (BlendLerp < UE_KINDA_SMALL_NUMBER && DeltaTime > UE_KINDA_SMALL_NUMBER)
+	if (Source.GetLinkNode())
 	{
-		// resetting the blendstack if there's no BlendLerp weight to it
-		Reset();
+		// updating blend stack
+		if (BlendLerp < UE_KINDA_SMALL_NUMBER && DeltaTime > UE_KINDA_SMALL_NUMBER)
+		{
+			// resetting the blendstack if there's no BlendLerp weight to it
+			Reset();
+		}
+
+		FAnimationUpdateContext BlendStackContext = Context.FractionalWeightAndRootMotion(BlendLerp, BlendLerp);
+
+		// bypassing FAnimNode_BlendStack::UpdateAssetPlayer, since we overridden its behaviour
+		FAnimNode_BlendStack_Standalone::UpdateAssetPlayer(BlendStackContext);
+
+	}
+	else
+	{
+		// bypassing FAnimNode_BlendStack::UpdateAssetPlayer, since we overridden its behaviour
+		FAnimNode_BlendStack_Standalone::UpdateAssetPlayer(Context);
 	}
-	FAnimationUpdateContext BlendStackContext = Context.FractionalWeightAndRootMotion(BlendLerp, BlendLerp);
-	Super::UpdateAssetPlayer(BlendStackContext);
 
 #if ENABLE_ANIM_DEBUG
 	if (UE_TRACE_CHANNELEXPR_IS_ENABLED(AnimationChannel))
 	{
-		TRACE_ANIM_NODE_VALUE(Context, *FString("Interacting"), bInteracting);
-		TRACE_ANIM_NODE_VALUE(Context, *FString("ExecuteBlendTo"), bExecuteBlendTo);
+		TRACE_ANIM_NODE_VALUE(Context, *FString("WasInteracting"), bWasInteracting);
+		TRACE_ANIM_NODE_VALUE(Context, *FString("IsInteracting"), bIsInteracting);
+		TRACE_ANIM_NODE_VALUE(Context, *FString("BlendToExecuted"), bBlendToExecuted);
 		TRACE_ANIM_NODE_VALUE(Context, *FString("BlendLerp"), BlendLerp);
 		TRACE_ANIM_NODE_VALUE(Context, *FString("TranslationWarpLerp"), TranslationWarpLerp);
 		TRACE_ANIM_NODE_VALUE(Context, *FString("RotationWarpLerp"), RotationWarpLerp);
 	}
 #endif // ENABLE_ANIM_DEBUG
+
+	bWasInteracting |= bIsInteracting;
 }
 
 void FAnimNode_Chimera::Evaluate_AnyThread(FPoseContext& Output)
diff --git a/Engine/Plugins/Experimental/Animation/Chimera/Source/Runtime/Private/ChimeraAsset.cpp b/Engine/Plugins/Experimental/Animation/Chimera/Source/Runtime/Private/ChimeraAsset.cpp
--- a/Engine/Plugins/Experimental/Animation/Chimera/Source/Runtime/Private/ChimeraAsset.cpp
+++ b/Engine/Plugins/Experimental/Animation/Chimera/Source/Runtime/Private/ChimeraAsset.cpp
@@ -6,6 +6,13 @@
 #include "DrawDebugHelpers.h"
 #include "PoseSearch/PoseSearchAssetSampler.h"
 
+namespace UE::Chimera
+{
+#if ENABLE_DRAW_DEBUG && ENABLE_ANIM_DEBUG
+TAutoConsoleVariable<bool> CVarAnimChimeraDebugDrawWarpTransforms(TEXT("a.Chimera.DebugDrawWarpTransforms"), false, TEXT("Enable / Disable Warp Transforms DebugDraw"));
+#endif
+}
+
 bool UChimeraAsset::IsLooping() const
 {
 	float CommonPlayLength = -1.f;
@@ -334,11 +341,15 @@ void UChimeraAsset::CalculateWarpTransforms(float Time, const TArrayView<const F
 	const FVector AssetReferencePosition = FindReferencePosition(AssetRootBoneTransforms, NormalizedWarpingWeightTranslation);
 	const FVector ActorsReferencePosition = FindReferencePosition(ActorRootBoneTransforms, NormalizedWarpingWeightTranslation);
 
-#if ENABLE_DRAW_DEBUG
-	DrawDebugCoordinateSystem(DebugDrawWorld, AssetReferencePosition, AssetReferenceOrientation.Rotator(), 10.f, false, 0.f, SDPG_Foreground);
-	DrawDebugCoordinateSystem(DebugDrawWorld, ActorsReferencePosition, ActorsReferenceOrientation.Rotator(), 10.f, false, 0.f, SDPG_Foreground);
-	DrawDebugCoordinateSystem(DebugDrawWorld, ActorsReferencePosition, WeightedActorsReferenceOrientation.Rotator(), 20.f, false, 0.f, SDPG_Foreground);
-#endif // ENABLE_DRAW_DEBUG
+#if ENABLE_DRAW_DEBUG && ENABLE_ANIM_DEBUG
+	const bool bDrawWarpTransforms = UE::Chimera::CVarAnimChimeraDebugDrawWarpTransforms.GetValueOnAnyThread();
+	if (bDrawWarpTransforms)
+	{
+		DrawDebugCoordinateSystem(DebugDrawWorld, AssetReferencePosition, AssetReferenceOrientation.Rotator(), 10.f, false, 0.f, SDPG_Foreground);
+		DrawDebugCoordinateSystem(DebugDrawWorld, ActorsReferencePosition, ActorsReferenceOrientation.Rotator(), 10.f, false, 0.f, SDPG_Foreground);
+		DrawDebugCoordinateSystem(DebugDrawWorld, ActorsReferencePosition, WeightedActorsReferenceOrientation.Rotator(), 20.f, false, 0.f, SDPG_Foreground);
+	}
+#endif // ENABLE_DRAW_DEBUG && ENABLE_ANIM_DEBUG
 
 	// aligning all the actors to ActorsReferencePosition, WeightedActorsReferenceOrientation
 	const FTransform AssetReferenceTransform(AssetReferenceOrientation, AssetReferencePosition);
@@ -349,8 +360,11 @@ void UChimeraAsset::CalculateWarpTransforms(float Time, const TArrayView<const F
 	{
 		FullAlignedActorRootBoneTransforms[ItemIndex] = (AssetRootBoneTransforms[ItemIndex] * AssetReferenceInverseTransform) * ActorsReferenceTransform;
 
-#if ENABLE_DRAW_DEBUG
-		DrawDebugCoordinateSystem(DebugDrawWorld, FullAlignedActorRootBoneTransforms[ItemIndex].GetTranslation(), FullAlignedActorRootBoneTransforms[ItemIndex].Rotator(), 10.f, false, 0.f, SDPG_Foreground);
-#endif // ENABLE_DRAW_DEBUG
+#if ENABLE_DRAW_DEBUG && ENABLE_ANIM_DEBUG
+		if (bDrawWarpTransforms)
+		{
+			DrawDebugCoordinateSystem(DebugDrawWorld, FullAlignedActorRootBoneTransforms[ItemIndex].GetTranslation(), FullAlignedActorRootBoneTransforms[ItemIndex].Rotator(), 10.f, false, 0.f, SDPG_Foreground);
+		}
+#endif // ENABLE_DRAW_DEBUG && ENABLE_ANIM_DEBUG
 	}
 }
\ No newline at end of file
diff --git a/Engine/Plugins/Experimental/Animation/Chimera/Source/Runtime/Public/Chimera/AnimNode_Chimera.h b/Engine/Plugins/Experimental/Animation/Chimera/Source/Runtime/Public/Chimera/AnimNode_Chimera.h
--- a/Engine/Plugins/Experimental/Animation/Chimera/Source/Runtime/Public/Chimera/AnimNode_Chimera.h
+++ b/Engine/Plugins/Experimental/Animation/Chimera/Source/Runtime/Public/Chimera/AnimNode_Chimera.h
@@ -9,14 +9,33 @@
 struct FAnimationInitializeContext;
 struct FNodeDebugData;
 
+UENUM()
+enum class EChimeraEvaluationMode
+{
+	// Chimera Node will continuously provide its availabilities and eventually blend to newly selected animations
+	ContinuousReselection,
+
+	// Chimera Node will continuously provide its availabilities to keep the interaction alive, but will play only the first selected animation
+	// the idea is to let the animation play until the end and allow the evenutual state machine playing this node to be able to perform an automatic transition
+	SingleSelection,
+
+	// @todo: is this needed?
+	// Chimera Node will stop providing its availabilities and consequently kill the interaction, when the first selected animation stop playing as continuing pose
+	// the idea is to let the animation play until its valid as continuing pose and allow the evenutual state machine playing this node to be able to perform an automatic transition
+	// UntilContinuingPoseIsValid,
+};
+
 USTRUCT(BlueprintInternalUseOnly, Experimental)
-struct CHIMERA_API FAnimNode_Chimera : public FAnimNode_BlendStack_Standalone
+struct CHIMERA_API FAnimNode_Chimera : public FAnimNode_BlendStack
 {
 	GENERATED_BODY();
 
-	UPROPERTY(EditAnywhere, BlueprintReadWrite, Category = Links, meta = (DisplayPriority = 0))
+	UPROPERTY(EditAnywhere, BlueprintReadWrite, Category = Links)
 	FPoseLink Source;
 
+	UPROPERTY(EditAnywhere, Category = Chimera)
+	EChimeraEvaluationMode EvaluationMode = EChimeraEvaluationMode::ContinuousReselection;
+
 	UPROPERTY(EditAnywhere, Category = Chimera, meta = (PinHiddenByDefault))
 	TArray<FChimeraAvailability> Availabilities;
 
@@ -31,24 +50,6 @@ struct CHIMERA_API FAnimNode_Chimera : public FAnimNode_BlendStack_Standalone
 	UPROPERTY(EditAnywhere, BlueprintReadWrite, Category = Warping, meta = (PinHiddenByDefault, ClampMin = "0"))
 	float InitialRotationWarpTime = 0.2f;
 
-	// tunable animation transition blend time 
-	UPROPERTY(EditAnywhere, BlueprintReadWrite, Category = Blending, meta = (PinHiddenByDefault, ClampMin = "0"))
-	float BlendTime = 0.2f;
-
-	UPROPERTY(EditAnywhere, BlueprintReadWrite, Category = Blending, meta = (PinHiddenByDefault, UseAsBlendProfile = true))
-	TObjectPtr<UBlendProfile> BlendProfile;
-
-	UPROPERTY(EditAnywhere, BlueprintReadWrite, Category = Blending, meta = (PinHiddenByDefault))
-	EAlphaBlendOption BlendOption = EAlphaBlendOption::Linear;
-
-	// tunable animation transition blend time 
-	UPROPERTY(EditAnywhere, Category = Blending)
-	bool bUseInertialBlend = false;
-
-	// Reset the blend stack if it has become relevant to the graph after not being updated on previous frames.
-	UPROPERTY(EditAnywhere, Category = Blending)
-	bool bResetOnBecomingRelevant = true;
-
 	// FAnimNode_Base interface
 	virtual void GatherDebugData(FNodeDebugData& DebugData) override;
 	virtual void Initialize_AnyThread(const FAnimationInitializeContext& Context) override;
@@ -58,10 +59,13 @@ struct CHIMERA_API FAnimNode_Chimera : public FAnimNode_BlendStack_Standalone
 	virtual void UpdateAssetPlayer(const FAnimationUpdateContext& Context) override;
 	// End of FAnimNode_Base interface
 
+	virtual void Reset() override;
+
 protected:
 	float BlendLerp = 0.f;
 	float TranslationWarpLerp = 0.f;
 	float RotationWarpLerp = 0.f;
+	bool bWasInteracting = false;
 
 	// Update Counter for detecting being relevant
 	FGraphTraversalCounter UpdateCounter;
