chimera - adding BlueprintCallable ChimeraQuery UFUNCTION, exposing bIsContinuingPoseSearch in FPoseSearchBlueprintResult

diff --git a/Engine/Plugins/Animation/PoseSearch/Source/Runtime/Private/MotionMatchingAnimNodeLibrary.cpp b/Engine/Plugins/Animation/PoseSearch/Source/Runtime/Private/MotionMatchingAnimNodeLibrary.cpp
--- a/Engine/Plugins/Animation/PoseSearch/Source/Runtime/Private/MotionMatchingAnimNodeLibrary.cpp
+++ b/Engine/Plugins/Animation/PoseSearch/Source/Runtime/Private/MotionMatchingAnimNodeLibrary.cpp
@@ -35,6 +35,7 @@ void UMotionMatchingAnimNodeLibrary::GetMotionMatchingSearchResult(const FMotion
 				{
 					Result.SelectedAnimation = DatabaseAsset->GetAnimationAsset();
 					Result.SelectedTime = MotionMatchingState.CurrentSearchResult.AssetTime;
+					Result.bIsContinuingPoseSearch = MotionMatchingState.CurrentSearchResult.bIsContinuingPoseSearch;
 					Result.bLoop = SearchIndexAsset->IsLooping();
 					Result.bIsMirrored = SearchIndexAsset->IsMirrored();
 					Result.BlendParameters = SearchIndexAsset->GetBlendParameters();
diff --git a/Engine/Plugins/Animation/PoseSearch/Source/Runtime/Private/PoseSearchLibrary.cpp b/Engine/Plugins/Animation/PoseSearch/Source/Runtime/Private/PoseSearchLibrary.cpp
--- a/Engine/Plugins/Animation/PoseSearch/Source/Runtime/Private/PoseSearchLibrary.cpp
+++ b/Engine/Plugins/Animation/PoseSearch/Source/Runtime/Private/PoseSearchLibrary.cpp
@@ -698,6 +698,7 @@ void UPoseSearchLibrary::MotionMatch(
 
 	Result.SelectedAnimation = nullptr;
 	Result.SelectedTime = 0.f;
+	Result.bIsContinuingPoseSearch = false;
 	Result.bLoop = false;
 	Result.bIsMirrored = false;
 	Result.BlendParameters = FVector::ZeroVector;
@@ -751,6 +752,7 @@ void UPoseSearchLibrary::MotionMatch(
 		{
 			Result.SelectedAnimation = DatabaseAsset->GetAnimationAsset();
 			Result.SelectedTime = SearchResult.AssetTime;
+			Result.bIsContinuingPoseSearch = SearchResult.bIsContinuingPoseSearch;
 			Result.bLoop = SearchIndexAsset->IsLooping();
 			Result.bIsMirrored = SearchIndexAsset->IsMirrored();
 			Result.BlendParameters = SearchIndexAsset->GetBlendParameters();
diff --git a/Engine/Plugins/Animation/PoseSearch/Source/Runtime/Public/PoseSearch/PoseSearchResult.h b/Engine/Plugins/Animation/PoseSearch/Source/Runtime/Public/PoseSearch/PoseSearchResult.h
--- a/Engine/Plugins/Animation/PoseSearch/Source/Runtime/Public/PoseSearch/PoseSearchResult.h
+++ b/Engine/Plugins/Animation/PoseSearch/Source/Runtime/Public/PoseSearch/PoseSearchResult.h
@@ -55,6 +55,9 @@ public:
 	UPROPERTY(Transient, VisibleAnywhere, BlueprintReadOnly, Category=State)
 	float SelectedTime = 0.f;
 	
+	UPROPERTY(Transient, VisibleAnywhere, BlueprintReadOnly, Category=State)
+	bool bIsContinuingPoseSearch = false;
+	
 	UPROPERTY(Transient, VisibleAnywhere, BlueprintReadOnly, Category=State)
 	float WantedPlayRate = 0.f;
 
diff --git a/Engine/Plugins/Experimental/Animation/Chimera/Source/Runtime/Private/AnimNode_Chimera.cpp b/Engine/Plugins/Experimental/Animation/Chimera/Source/Runtime/Private/AnimNode_Chimera.cpp
--- a/Engine/Plugins/Experimental/Animation/Chimera/Source/Runtime/Private/AnimNode_Chimera.cpp
+++ b/Engine/Plugins/Experimental/Animation/Chimera/Source/Runtime/Private/AnimNode_Chimera.cpp
@@ -203,14 +203,10 @@ void FAnimNode_Chimera::Evaluate_AnyThread(FPoseContext& Output)
 {
 	check(Output.AnimInstanceProxy);
 
-	// @todo: do we still need to evaluate Source if blendstack weight is 100%?
 	// evaluating Source to get the base pose 
 	Source.Evaluate(Output);
 
-	if (!AnimPlayers.IsEmpty())
-	{
-		Super::Evaluate_AnyThread(Output);
-	}
+	Super::Evaluate_AnyThread(Output);
 
 	if (TranslationWarpLerp > UE_KINDA_SMALL_NUMBER || RotationWarpLerp > UE_KINDA_SMALL_NUMBER)
 	{
diff --git a/Engine/Plugins/Experimental/Animation/Chimera/Source/Runtime/Private/ChimeraLibrary.cpp b/Engine/Plugins/Experimental/Animation/Chimera/Source/Runtime/Private/ChimeraLibrary.cpp
--- a/Engine/Plugins/Experimental/Animation/Chimera/Source/Runtime/Private/ChimeraLibrary.cpp
+++ b/Engine/Plugins/Experimental/Animation/Chimera/Source/Runtime/Private/ChimeraLibrary.cpp
@@ -14,6 +14,11 @@ FChimeraBlueprintResult UChimeraLibrary::ChimeraQuery_Pure(TArray<FChimeraAvaila
 	return Result;
 }
 
+FChimeraBlueprintResult UChimeraLibrary::ChimeraQuery(TArray<FChimeraAvailability> Availabilities, UObject* AnimInstance, FName PoseHistoryName, bool bValidateResultAgainstAvailabilities)
+{
+	return ChimeraQuery_Pure(Availabilities, AnimInstance, PoseHistoryName, bValidateResultAgainstAvailabilities);
+}
+
 FChimeraBlueprintResult UChimeraLibrary::ChimeraQuery(const TArrayView<const FChimeraAvailability> Availabilities, UObject* AnimInstance, const FAnimNode_PoseSearchHistoryCollector_Base* HistoryCollector, bool bValidateResultAgainstAvailabilities)
 {
 	FChimeraBlueprintResult Result;
diff --git a/Engine/Plugins/Experimental/Animation/Chimera/Source/Runtime/Public/Chimera/ChimeraLibrary.h b/Engine/Plugins/Experimental/Animation/Chimera/Source/Runtime/Public/Chimera/ChimeraLibrary.h
--- a/Engine/Plugins/Experimental/Animation/Chimera/Source/Runtime/Public/Chimera/ChimeraLibrary.h
+++ b/Engine/Plugins/Experimental/Animation/Chimera/Source/Runtime/Public/Chimera/ChimeraLibrary.h
@@ -110,6 +110,10 @@ public:
 	UFUNCTION(BlueprintPure, Category = "Animation|Chimera", meta = (BlueprintThreadSafe))
 	static FChimeraBlueprintResult ChimeraQuery_Pure(TArray<FChimeraAvailability> Availabilities, UObject* AnimInstance, FName PoseHistoryName, bool bValidateResultAgainstAvailabilities);
 
+	// BlueprintCallable version of ChimeraQuery_Pure
+	UFUNCTION(BlueprintCallable, Category = "Animation|Chimera", meta = (BlueprintThreadSafe))
+	static FChimeraBlueprintResult ChimeraQuery(TArray<FChimeraAvailability> Availabilities, UObject* AnimInstance, FName PoseHistoryName, bool bValidateResultAgainstAvailabilities);
+
 	static FChimeraBlueprintResult ChimeraQuery(const TArrayView<const FChimeraAvailability> Availabilities, UObject* AnimInstance, const FAnimNode_PoseSearchHistoryCollector_Base* HistoryCollector, bool bValidateResultAgainstAvailabilities);
 
 };
\ No newline at end of file
