chimera - support for single character "interaction"

diff --git a/Engine/Plugins/Experimental/Animation/Chimera/Source/Runtime/Private/ChimeraIslandComponent.cpp b/Engine/Plugins/Experimental/Animation/Chimera/Source/Runtime/Private/ChimeraIslandComponent.cpp
--- a/Engine/Plugins/Experimental/Animation/Chimera/Source/Runtime/Private/ChimeraIslandComponent.cpp
+++ b/Engine/Plugins/Experimental/Animation/Chimera/Source/Runtime/Private/ChimeraIslandComponent.cpp
@@ -97,7 +97,7 @@ namespace UE::Chimera
 		}
 
 		const int32 Num = AnimInstances.Num();
-		if (Num < 2)
+		if (Num < 1)
 		{
 			return false;
 		}
@@ -278,21 +278,10 @@ void UChimeraIslandComponent::AddSearchContext(const UE::Chimera::FSearchContext
 	SearchContexts.Add(SearchContext);
 }
 
-void UChimeraIslandComponent::ResetSearchContexts()
+void UChimeraIslandComponent::Uninject()
 {
 	check(IsInGameThread());
-	SearchContexts.Reset();
-}
-
-void UChimeraIslandComponent::ResetSearchResults()
-{
-	check(IsInGameThread());
-	SearchResults.Reset();
-	bSearchPerfomed = false;
-}
 
-void UChimeraIslandComponent::UninjectFromAllActors()
-{
 	// Called by UChimeraSubsystem::Tick when there aren't animation jobs flying. No need to FScopeLock Lock(&Mutex);
 	for (TWeakObjectPtr<UCharacterMovementComponent>& CharacterMovementComponentPtr : CharacterMovementComponents)
 	{
@@ -306,6 +295,10 @@ void UChimeraIslandComponent::UninjectFromAllActors()
 
 	CharacterMovementComponents.Reset();
 	SkeletalMeshComponents.Reset();
+
+	SearchContexts.Reset();
+	SearchResults.Reset();
+	bSearchPerfomed = false;
 }
 
 bool UChimeraIslandComponent::IsUninjected()
@@ -389,23 +382,35 @@ bool UChimeraIslandComponent::DoSearch_AnyThread(UObject* AnimInstance, FChimera
 			}
 		}
 
+		// making sure we called Uninject()
+		check(SearchResults.IsEmpty());
+
 		////////////////////////////////////////////////////////////////////////////////////////////////////////////////
 		// WIP!
-		// @todo: figure out multiple policies to sue the most characters? right now only the best search is "valid"
+		// @todo: figure out multiple policies to use the most characters? right now only the best search is "valid" with the most characters
 		if (!PoseSearchResults.IsEmpty())
 		{
 			// locking to update SearchResults and bSearchPerfomed
 			FScopeLock Lock(&SearchResultsMutex);
 
-			FPoseSearchCost BestPoseCost;
 			int32 BestSearchIndex = INDEX_NONE;
 			for (int32 SearchIndex = 0; SearchIndex < PoseSearchResults.Num(); ++SearchIndex)
 			{
-				const UE::PoseSearch::FSearchResult& PoseSearchResult = PoseSearchResults[SearchIndex];
-				if (PoseSearchResult.IsValid() && PoseSearchResult.PoseCost < BestPoseCost)
+				if (PoseSearchResults[SearchIndex].IsValid())
 				{
-					BestPoseCost = PoseSearchResult.PoseCost;
-					BestSearchIndex = SearchIndex;
+					if (BestSearchIndex == INDEX_NONE)
+					{
+						BestSearchIndex = SearchIndex;
+					}
+					else if (SearchContexts[SearchIndex].Roles.Num() > SearchContexts[BestSearchIndex].Roles.Num())
+					{
+						BestSearchIndex = SearchIndex;
+					}
+					else if (SearchContexts[SearchIndex].Roles.Num() == SearchContexts[BestSearchIndex].Roles.Num() &&
+						PoseSearchResults[SearchIndex].PoseCost < PoseSearchResults[BestSearchIndex].PoseCost)
+					{
+						BestSearchIndex = SearchIndex;
+					}
 				}
 			}
 
diff --git a/Engine/Plugins/Experimental/Animation/Chimera/Source/Runtime/Private/ChimeraSubsystem.cpp b/Engine/Plugins/Experimental/Animation/Chimera/Source/Runtime/Private/ChimeraSubsystem.cpp
--- a/Engine/Plugins/Experimental/Animation/Chimera/Source/Runtime/Private/ChimeraSubsystem.cpp
+++ b/Engine/Plugins/Experimental/Animation/Chimera/Source/Runtime/Private/ChimeraSubsystem.cpp
@@ -258,7 +258,7 @@ UChimeraIslandComponent& UChimeraSubsystem::GetAvailableIsland()
 void UChimeraSubsystem::DestroyIsland(int32 Index)
 {
 	UChimeraIslandComponent& Island = *Islands[Index].Get();
-	Island.UninjectFromAllActors();
+	Island.Uninject();
 	Island.RemoveFromRoot();
 	Island.UnregisterComponent();
 	Island.DestroyComponent();
@@ -279,10 +279,7 @@ void UChimeraSubsystem::UninjectAllIslands()
 
 	for (TObjectPtr<UChimeraIslandComponent>& IslandPtr : Islands)
 	{
-		UChimeraIslandComponent& Island = *IslandPtr.Get();
-		Island.UninjectFromAllActors();
-		Island.ResetSearchContexts();
-		Island.ResetSearchResults();
+		IslandPtr.Get()->Uninject();
 	}
 }
 
diff --git a/Engine/Plugins/Experimental/Animation/Chimera/Source/Runtime/Public/Chimera/ChimeraIslandComponent.h b/Engine/Plugins/Experimental/Animation/Chimera/Source/Runtime/Public/Chimera/ChimeraIslandComponent.h
--- a/Engine/Plugins/Experimental/Animation/Chimera/Source/Runtime/Public/Chimera/ChimeraIslandComponent.h
+++ b/Engine/Plugins/Experimental/Animation/Chimera/Source/Runtime/Public/Chimera/ChimeraIslandComponent.h
@@ -59,12 +59,10 @@ private:
 	virtual void TickComponent(float DeltaTime, enum ELevelTick TickType, FActorComponentTickFunction *ThisTickFunction) override;
 
 	void InjectToActor(AActor* Actor);
-	void UninjectFromAllActors();
+	void Uninject();
 	bool IsUninjected();
 
 	void AddSearchContext(const UE::Chimera::FSearchContext& SearchContext);
-	void ResetSearchContexts();
-	void ResetSearchResults();
 
 	bool GetResult_AnyThread(UObject* AnimInstance, FChimeraBlueprintResult& Result);
 
