chimera
- threading issue fix. added PostTickFunction
- removed dependencies from ChimeraAsset in favor of the more generic MultiAnimAsset
- replaced UChimeraIslandComponent with the lightweight UE::Chimera::FIsland
- added FPoseSearchContinuingProperties to UChimeraSubsystem::Query_AnyThread

diff --git a/Engine/Plugins/Animation/PoseSearch/Source/Runtime/Private/PoseSearchLibrary.cpp b/Engine/Plugins/Animation/PoseSearch/Source/Runtime/Private/PoseSearchLibrary.cpp
--- a/Engine/Plugins/Animation/PoseSearch/Source/Runtime/Private/PoseSearchLibrary.cpp
+++ b/Engine/Plugins/Animation/PoseSearch/Source/Runtime/Private/PoseSearchLibrary.cpp
@@ -927,20 +927,40 @@ UE::PoseSearch::FSearchResult UPoseSearchLibrary::MotionMatch(
 				{
 					const FSearchIndexAsset& SearchIndexAsset = SearchIndex.Assets[AssetIndex];
 
+					const float FirstSampleTime = SearchIndexAsset.GetFirstSampleTime(Database->Schema->SampleRate);
+					const float LastSampleTime = SearchIndexAsset.GetLastSampleTime(Database->Schema->SampleRate);
+
 					bool bCanAdvance = true;
-					if (!SearchIndexAsset.IsLooping())
+					if (SearchIndexAsset.IsLooping())
+					{
+						const float DeltaSampleTime = LastSampleTime - FirstSampleTime;
+						if (DeltaSampleTime < UE_SMALL_NUMBER)
+						{
+							ReconstructedPreviousSearchResult.AssetTime = FirstSampleTime;
+						}
+						else if (ContinuingProperties.PlayingAssetAccumulatedTime < FirstSampleTime)
+						{
+							ReconstructedPreviousSearchResult.AssetTime = FMath::Fmod(ContinuingProperties.PlayingAssetAccumulatedTime - FirstSampleTime, DeltaSampleTime) + DeltaSampleTime + FirstSampleTime;
+						}
+						else if (ContinuingProperties.PlayingAssetAccumulatedTime > LastSampleTime)
+						{
+							ReconstructedPreviousSearchResult.AssetTime = FMath::Fmod(ContinuingProperties.PlayingAssetAccumulatedTime - FirstSampleTime, DeltaSampleTime) + FirstSampleTime;
+						}
+						else
+						{
+							ReconstructedPreviousSearchResult.AssetTime = ContinuingProperties.PlayingAssetAccumulatedTime;
+						}
+					}
+					else
 					{
-						const float FirstSampleTime = SearchIndexAsset.GetFirstSampleTime(Database->Schema->SampleRate);
-						const float LastSampleTime = SearchIndexAsset.GetLastSampleTime(Database->Schema->SampleRate) - DeltaSeconds;
 						const float MaxTimeToBeAbleToContinuingPlayingAnimation = LastSampleTime - DeltaSeconds;
-
 						bCanAdvance = ContinuingProperties.PlayingAssetAccumulatedTime >= FirstSampleTime && ContinuingProperties.PlayingAssetAccumulatedTime < MaxTimeToBeAbleToContinuingPlayingAnimation;
+						ReconstructedPreviousSearchResult.AssetTime = ContinuingProperties.PlayingAssetAccumulatedTime;
 					}
 
 					if (bCanAdvance)
 					{
 						ReconstructedPreviousSearchResult.Database = Database;
-						ReconstructedPreviousSearchResult.AssetTime = ContinuingProperties.PlayingAssetAccumulatedTime;
 						ReconstructedPreviousSearchResult.PoseIdx = Database->GetPoseIndexFromTime(ContinuingProperties.PlayingAssetAccumulatedTime, SearchIndexAsset);
 						SearchContext.UpdateCurrentResultPoseVector();
 
diff --git a/Engine/Plugins/Animation/PoseSearch/Source/Runtime/Public/PoseSearch/MultiAnimAsset.h b/Engine/Plugins/Animation/PoseSearch/Source/Runtime/Public/PoseSearch/MultiAnimAsset.h
--- a/Engine/Plugins/Animation/PoseSearch/Source/Runtime/Public/PoseSearch/MultiAnimAsset.h
+++ b/Engine/Plugins/Animation/PoseSearch/Source/Runtime/Public/PoseSearch/MultiAnimAsset.h
@@ -28,6 +28,8 @@ public:
 	[[nodiscard]] virtual UAnimationAsset* GetAnimationAsset(const FName& Role) const PURE_VIRTUAL(UMultiAnimAsset::GetAnimationAsset, return nullptr;);
 	[[nodiscard]] virtual FTransform GetOrigin(const FName& Role) const PURE_VIRTUAL(UMultiAnimAsset::GetOrigin, return FTransform::Identity;);
 
+	virtual void CalculateWarpTransforms(float Time, const TArrayView<const FTransform> ActorRootBoneTransforms, TArrayView<FTransform> FullAlignedActorRootBoneTransforms, const UWorld* DebugDrawWorld = nullptr) const PURE_VIRTUAL(UMultiAnimAsset::CalculateWarpTransforms, );
+
 	UFUNCTION(BlueprintPure, Category = "Animation", meta=(BlueprintThreadSafe, DisplayName = "Get Animation Asset"))
 	UAnimationAsset* BP_GetAnimationAsset(const FName& Role) const { return GetAnimationAsset(Role); }
 
diff --git a/Engine/Plugins/Animation/PoseSearch/Source/Runtime/Public/PoseSearch/PoseSearchCost.h b/Engine/Plugins/Animation/PoseSearch/Source/Runtime/Public/PoseSearch/PoseSearchCost.h
--- a/Engine/Plugins/Animation/PoseSearch/Source/Runtime/Public/PoseSearch/PoseSearchCost.h
+++ b/Engine/Plugins/Animation/PoseSearch/Source/Runtime/Public/PoseSearch/PoseSearchCost.h
@@ -23,6 +23,7 @@ public:
 	bool IsValid() const { return IsCostValid(TotalCost); }
 	float GetTotalCost() const { return TotalCost; }
 	bool operator<(const FPoseSearchCost& Other) const { return TotalCost < Other.TotalCost; }
+	bool operator==(const FPoseSearchCost& Other) const { return TotalCost == Other.TotalCost; }
 
 protected:
 	UPROPERTY()
diff --git a/Engine/Plugins/Animation/PoseSearch/Source/Runtime/Public/PoseSearch/PoseSearchResult.h b/Engine/Plugins/Animation/PoseSearch/Source/Runtime/Public/PoseSearch/PoseSearchResult.h
--- a/Engine/Plugins/Animation/PoseSearch/Source/Runtime/Public/PoseSearch/PoseSearchResult.h
+++ b/Engine/Plugins/Animation/PoseSearch/Source/Runtime/Public/PoseSearch/PoseSearchResult.h
@@ -39,6 +39,17 @@ struct FSearchResult
 	POSESEARCH_API const FSearchIndexAsset* GetSearchIndexAsset(bool bMandatory = false) const;
 	
 	bool CanAdvance(float DeltaTime) const;
+
+	bool operator==(const FSearchResult& Other) const
+	{
+		// best cost of the currently selected PoseIdx (it could be equal to ContinuingPoseCost)
+		return	PoseCost == Other.PoseCost &&
+			PoseIdx == Other.PoseIdx &&
+			Database == Other.Database &&
+			AssetTime == Other.AssetTime &&
+			bIsContinuingPoseSearch == Other.bIsContinuingPoseSearch;
+	}
+
 };
 
 } // namespace UE::PoseSearch
diff --git a/Engine/Plugins/Experimental/Animation/Chimera/Source/Editor/Public/AnimGraphNode_Chimera.h b/Engine/Plugins/Experimental/Animation/Chimera/Source/Editor/Public/AnimGraphNode_Chimera.h
--- a/Engine/Plugins/Experimental/Animation/Chimera/Source/Editor/Public/AnimGraphNode_Chimera.h
+++ b/Engine/Plugins/Experimental/Animation/Chimera/Source/Editor/Public/AnimGraphNode_Chimera.h
@@ -6,7 +6,7 @@
 #include "Chimera/AnimNode_Chimera.h"
 #include "AnimGraphNode_Chimera.generated.h"
 
-UCLASS(Experimental)
+UCLASS()
 class CHIMERAEDITOR_API UAnimGraphNode_Chimera : public UAnimGraphNode_BlendStack_Base
 {
 	GENERATED_UCLASS_BODY()
diff --git a/Engine/Plugins/Experimental/Animation/Chimera/Source/Runtime/Chimera.Build.cs b/Engine/Plugins/Experimental/Animation/Chimera/Source/Runtime/Chimera.Build.cs
--- a/Engine/Plugins/Experimental/Animation/Chimera/Source/Runtime/Chimera.Build.cs
+++ b/Engine/Plugins/Experimental/Animation/Chimera/Source/Runtime/Chimera.Build.cs
@@ -6,6 +6,10 @@ public class Chimera : ModuleRules
 {
 	public Chimera(ReadOnlyTargetRules Target) : base(Target)
 	{
+		//OptimizeCode = CodeOptimization.Never;
+		//bUseUnity = false;
+		//PCHUsage = PCHUsageMode.NoPCHs;
+
 		PublicDependencyModuleNames.AddRange(
 			new string[]
 			{
@@ -17,7 +21,7 @@ public class Chimera : ModuleRules
 				"PoseSearch",
             });
 
-			PrivateDependencyModuleNames.AddRange(
+		PrivateDependencyModuleNames.AddRange(
 			new string[]
 			{
 				"BlendStack"
diff --git a/Engine/Plugins/Experimental/Animation/Chimera/Source/Runtime/Private/AnimNode_Chimera.cpp b/Engine/Plugins/Experimental/Animation/Chimera/Source/Runtime/Private/AnimNode_Chimera.cpp
--- a/Engine/Plugins/Experimental/Animation/Chimera/Source/Runtime/Private/AnimNode_Chimera.cpp
+++ b/Engine/Plugins/Experimental/Animation/Chimera/Source/Runtime/Private/AnimNode_Chimera.cpp
@@ -8,6 +8,7 @@
 #include "Chimera/ChimeraLibrary.h"
 #include "PoseSearch/PoseHistoryProvider.h"
 #include "PoseSearch/PoseSearchDatabase.h"
+#include "PoseSearch/PoseSearchLibrary.h"
 #include "PoseSearch/PoseSearchSchema.h"
 #include "HAL/IConsoleManager.h"
 
@@ -45,7 +46,7 @@ void FAnimNode_Chimera::Reset()
 	BlendLerp = 0.f;
 	TranslationWarpLerp = 0.f;
 	RotationWarpLerp = 0.f;
-	bWasInteracting = false;
+	InteractingRolesNum = 0;
 }
 
 void FAnimNode_Chimera::UpdateAssetPlayer(const FAnimationUpdateContext& Context)
@@ -63,7 +64,7 @@ void FAnimNode_Chimera::UpdateAssetPlayer(const FAnimationUpdateContext& Context
 
 	bool bBlendToExecuted = ConditionalBlendTo(Context);
 
-	bool bIsInteracting = false;
+	int32 NewInteractingRolesNum = 0;
 	const float DeltaTime = Context.GetDeltaTime();
 	FPoseHistoryProvider* PoseHistoryProvider = Context.GetMessage<FPoseHistoryProvider>();
 	if (!PoseHistoryProvider)
@@ -72,21 +73,33 @@ void FAnimNode_Chimera::UpdateAssetPlayer(const FAnimationUpdateContext& Context
 	}
 	else
 	{
+		// we're NOT providing the UMultiAnimAsset here, relying on the Chimera subsystem to resolve the ContinuingProperties properly
+		FPoseSearchContinuingProperties ContinuingProperties;
+		ContinuingProperties.PlayingAsset = GetAnimAsset();
+		ContinuingProperties.PlayingAssetAccumulatedTime = GetAccumulatedTime();
+
 		check(Context.AnimInstanceProxy);
-		const FChimeraBlueprintResult Result = UChimeraLibrary::ChimeraQuery(Availabilities, Context.AnimInstanceProxy->GetAnimInstanceObject(), PoseHistoryProvider->GetHistoryCollector(), bValidateResultAgainstAvailabilities);
+		const FChimeraBlueprintResult Result = UChimeraLibrary::ChimeraQuery(Availabilities, Context.AnimInstanceProxy->GetAnimInstanceObject(), ContinuingProperties, PoseHistoryProvider->GetHistoryCollector(), bValidateResultAgainstAvailabilities);
 
+		UAnimationAsset* RoledAnimAsset = nullptr;
 		if (const UMultiAnimAsset* MultiAnimAsset = Cast<UMultiAnimAsset>(Result.SelectedAnimation))
 		{
 			check(Result.SelectedDatabase != nullptr);
 
-			UAnimationAsset* RoledAnimationAsset = MultiAnimAsset->GetAnimationAsset(Result.Role);
-			check(RoledAnimationAsset);
-
-			bIsInteracting = true;
+			RoledAnimAsset = MultiAnimAsset->GetAnimationAsset(Result.Role);
+			NewInteractingRolesNum = MultiAnimAsset->GetNumRoles();
+		}
+		else if (UAnimationAsset* SingleAnimAsset = Cast<UAnimationAsset>(Result.SelectedAnimation))
+		{
+			RoledAnimAsset = SingleAnimAsset;
+			NewInteractingRolesNum = 1;
+		}
 
+		if (RoledAnimAsset)
+		{
 			bool bExecuteBlendTo = false;
 			bool bUpdatePropertiesFromResult = false;
-			if (!bWasInteracting || AnimPlayers.IsEmpty())
+			if (InteractingRolesNum == 0 || AnimPlayers.IsEmpty())
 			{
 				bExecuteBlendTo = true;
 				bUpdatePropertiesFromResult = true;
@@ -94,9 +107,9 @@ void FAnimNode_Chimera::UpdateAssetPlayer(const FAnimationUpdateContext& Context
 			else if (EvaluationMode == EChimeraEvaluationMode::ContinuousReselection)
 			{
 				const FBlendStackAnimPlayer& MainAnimPlayer = AnimPlayers[0];
-				const UAnimationAsset* PlayingAnimationAsset = MainAnimPlayer.GetAnimationAsset();
+				const UAnimationAsset* PlayingAnimAsset = MainAnimPlayer.GetAnimationAsset();
 
-				if (RoledAnimationAsset != PlayingAnimationAsset ||
+				if (RoledAnimAsset != PlayingAnimAsset ||
 					Result.bIsMirrored != MainAnimPlayer.GetMirror() ||
 					Result.BlendParameters != MainAnimPlayer.GetBlendParameters() ||
 					!Result.bIsContinuingPoseSearch)
@@ -106,7 +119,7 @@ void FAnimNode_Chimera::UpdateAssetPlayer(const FAnimationUpdateContext& Context
 
 				bUpdatePropertiesFromResult = true;
 			}
-			else if (RoledAnimationAsset == AnimPlayers[0].GetAnimationAsset() && Result.bIsContinuingPoseSearch)
+			else if (RoledAnimAsset == AnimPlayers[0].GetAnimationAsset() && Result.bIsContinuingPoseSearch)
 			{
 				// we don't update FullAlignedActorRootBoneTransform since we're not planning to blend into the newly selected animation here
 				bUpdatePropertiesFromResult = true;
@@ -124,12 +137,11 @@ void FAnimNode_Chimera::UpdateAssetPlayer(const FAnimationUpdateContext& Context
 				const FPoseSearchRoledSkeleton* RoledSkeleton = Result.SelectedDatabase->Schema->GetRoledSkeleton(Result.Role);
 				check(RoledSkeleton);
 
-				BlendTo(Context, RoledAnimationAsset, Result.SelectedTime, Result.bLoop, Result.bIsMirrored, RoledSkeleton->MirrorDataTable.Get(),
+				BlendTo(Context, RoledAnimAsset, Result.SelectedTime, Result.bLoop, Result.bIsMirrored, RoledSkeleton->MirrorDataTable.Get(),
 					BlendTime, BlendProfile, BlendOption, bUseInertialBlend, BlendParameters, WantedPlayRate);
 
 				bBlendToExecuted = true;
 			}
-
 		}
 	}
 
@@ -140,18 +152,20 @@ void FAnimNode_Chimera::UpdateAssetPlayer(const FAnimationUpdateContext& Context
 	UpdateBlendspaceParameters(BlendspaceUpdateMode, BlendParameters);
 
 	// calculating the translation and rotation warp lerps, used to warp the root transform towards the last computed FullAlignedActorRootBoneTransform
-	const float Sign = bIsInteracting ? 1.f : -1.f;
 	if (BlendTime > UE_KINDA_SMALL_NUMBER)
 	{
-		BlendLerp = FMath::Clamp(BlendLerp + (Sign * DeltaTime / BlendTime), 0.f, 1.f);
+		const float BlendSign = NewInteractingRolesNum > 0 ? 1.f : -1.f;
+		BlendLerp = FMath::Clamp(BlendLerp + (BlendSign * DeltaTime / BlendTime), 0.f, 1.f);
 	}
+	
+	const float WarpSign = NewInteractingRolesNum > 1 ? 1.f : -1.f;
 	if (InitialTranslationWarpTime > UE_KINDA_SMALL_NUMBER)
 	{
-		TranslationWarpLerp = FMath::Clamp(TranslationWarpLerp + (Sign * DeltaTime / InitialTranslationWarpTime), 0.f, 1.f);
+		TranslationWarpLerp = FMath::Clamp(TranslationWarpLerp + (WarpSign * DeltaTime / InitialTranslationWarpTime), 0.f, 1.f);
 	}
 	if (InitialRotationWarpTime > UE_KINDA_SMALL_NUMBER)
 	{
-		RotationWarpLerp = FMath::Clamp(RotationWarpLerp + (Sign * DeltaTime / InitialRotationWarpTime), 0.f, 1.f);
+		RotationWarpLerp = FMath::Clamp(RotationWarpLerp + (WarpSign * DeltaTime / InitialRotationWarpTime), 0.f, 1.f);
 	}
 
 	// updating source
@@ -187,8 +201,8 @@ void FAnimNode_Chimera::UpdateAssetPlayer(const FAnimationUpdateContext& Context
 #if ENABLE_ANIM_DEBUG
 	if (UE_TRACE_CHANNELEXPR_IS_ENABLED(AnimationChannel))
 	{
-		TRACE_ANIM_NODE_VALUE(Context, *FString("WasInteracting"), bWasInteracting);
-		TRACE_ANIM_NODE_VALUE(Context, *FString("IsInteracting"), bIsInteracting);
+		TRACE_ANIM_NODE_VALUE(Context, *FString("InteractingRolesNum"), InteractingRolesNum);
+		TRACE_ANIM_NODE_VALUE(Context, *FString("NewInteractingRolesNum"), NewInteractingRolesNum);
 		TRACE_ANIM_NODE_VALUE(Context, *FString("BlendToExecuted"), bBlendToExecuted);
 		TRACE_ANIM_NODE_VALUE(Context, *FString("BlendLerp"), BlendLerp);
 		TRACE_ANIM_NODE_VALUE(Context, *FString("TranslationWarpLerp"), TranslationWarpLerp);
@@ -196,7 +210,7 @@ void FAnimNode_Chimera::UpdateAssetPlayer(const FAnimationUpdateContext& Context
 	}
 #endif // ENABLE_ANIM_DEBUG
 
-	bWasInteracting |= bIsInteracting;
+	InteractingRolesNum = NewInteractingRolesNum;
 }
 
 void FAnimNode_Chimera::Evaluate_AnyThread(FPoseContext& Output)
@@ -216,8 +230,8 @@ void FAnimNode_Chimera::Evaluate_AnyThread(FPoseContext& Output)
 #if ENABLE_ANIM_DEBUG
 		if (CVarAnimNodeChimeraDebug.GetValueOnAnyThread())
 		{
-			Output.AnimInstanceProxy->AnimDrawDebugCoordinateSystem(FullAlignedActorRootBoneTransform.GetLocation(), FullAlignedActorRootBoneTransform.Rotator(), 10.0f, false, 0.f, 0.f, SDPG_Foreground);
-			Output.AnimInstanceProxy->AnimDrawDebugCoordinateSystem(ComponentTransform.GetLocation(), ComponentTransform.Rotator(), 10.0f, false, 0.f, 0.f, SDPG_Foreground);
+			Output.AnimInstanceProxy->AnimDrawDebugCoordinateSystem(FullAlignedActorRootBoneTransform.GetLocation(), FullAlignedActorRootBoneTransform.Rotator(), 15.f, false, 0.f, 0.f, SDPG_Foreground);
+			Output.AnimInstanceProxy->AnimDrawDebugCoordinateSystem(ComponentTransform.GetLocation(), ComponentTransform.Rotator(), 5.f, false, 0.f, 0.f, SDPG_Foreground);
 		}
 #endif // ENABLE_ANIM_DEBUG
 
diff --git a/Engine/Plugins/Experimental/Animation/Chimera/Source/Runtime/Private/ChimeraAnimNodeLibrary.cpp b/Engine/Plugins/Experimental/Animation/Chimera/Source/Runtime/Private/ChimeraAnimNodeLibrary.cpp
new file mode 100644
--- /dev/null
+++ b/Engine/Plugins/Experimental/Animation/Chimera/Source/Runtime/Private/ChimeraAnimNodeLibrary.cpp
@@ -0,0 +1,22 @@
+// Copyright Epic Games, Inc. All Rights Reserved.
+
+#include "Chimera/ChimeraAnimNodeLibrary.h"
+#include "Chimera/AnimNode_Chimera.h"
+#include "Chimera/ChimeraDefines.h"
+
+FChimeraAnimNodeReference UChimeraAnimNodeLibrary::ConvertToChimeraNode(const FAnimNodeReference& Node, EAnimNodeReferenceConversionResult& Result)
+{
+	return FAnimNodeReference::ConvertToType<FChimeraAnimNodeReference>(Node, Result);
+}
+
+void UChimeraAnimNodeLibrary::SetAvailabilities(const FChimeraAnimNodeReference& ChimeraNode, const TArray<FChimeraAvailability>& Availabilities)
+{
+	if (FAnimNode_Chimera* ChimeraNodePtr = ChimeraNode.GetAnimNodePtr<FAnimNode_Chimera>())
+	{
+		ChimeraNodePtr->Availabilities = Availabilities;
+	}
+	else
+	{
+		UE_LOG(LogChimera, Warning, TEXT("UChimeraAnimNodeLibrary::SetAvailabilities called on an invalid context or with an invalid type"));
+	}
+}
\ No newline at end of file
diff --git a/Engine/Plugins/Experimental/Animation/Chimera/Source/Runtime/Private/ChimeraIslandComponent.cpp b/Engine/Plugins/Experimental/Animation/Chimera/Source/Runtime/Private/ChimeraIslandComponent.cpp
--- a/Engine/Plugins/Experimental/Animation/Chimera/Source/Runtime/Private/ChimeraIslandComponent.cpp
+++ b/Engine/Plugins/Experimental/Animation/Chimera/Source/Runtime/Private/ChimeraIslandComponent.cpp
@@ -4,210 +4,268 @@
 #include "Animation/AnimClassInterface.h"
 #include "Animation/AnimInstance.h"
 #include "Components/SkeletalMeshComponent.h"
-#include "Chimera/ChimeraAsset.h"
 #include "Chimera/ChimeraDefines.h"
 #include "Chimera/ChimeraLibrary.h"
 #include "DrawDebugHelpers.h"
 #include "GameFramework/CharacterMovementComponent.h"
 #include "PoseSearch/AnimNode_PoseSearchHistoryCollector.h"
+#include "PoseSearch/MultiAnimAsset.h"
 #include "PoseSearch/PoseSearchLibrary.h"
 #include "PoseSearch/PoseSearchDatabase.h"
 #include "PoseSearch/PoseSearchSchema.h"
 
 namespace UE::Chimera
 {
+
 #if ENABLE_ANIM_DEBUG
-	static TAutoConsoleVariable<bool> CVarChimeraShowIslands(TEXT("a.Chimera.ShowIslands"), false, TEXT("Show Chimera Islands"));
+static TAutoConsoleVariable<bool> CVarChimeraShowIslands(TEXT("a.Chimera.ShowIslands"), false, TEXT("Show Chimera Islands"));
 #endif
 
-	static UE::Chimera::FSearchResult InitSearchResult(const UE::PoseSearch::FSearchResult& SearchResult, int32 SearchIndex, const UE::Chimera::FSearchContext& SearchContext)
-	{
-		UE::Chimera::FSearchResult ChimeraSearchResult;
-		static_cast<UE::PoseSearch::FSearchResult&>(ChimeraSearchResult) = SearchResult;
-		ChimeraSearchResult.SearchIndex = SearchIndex;
+// Utils functions
+///////////////////////////////////////////////////////////
+static FSearchResult InitSearchResult(const UE::PoseSearch::FSearchResult& SearchResult, int32 SearchIndex, const FSearchContext& SearchContext)
+{
+	FSearchResult ChimeraSearchResult;
+	static_cast<UE::PoseSearch::FSearchResult&>(ChimeraSearchResult) = SearchResult;
+	ChimeraSearchResult.SearchIndex = SearchIndex;
 
-		// @todo: WIP! calculating warping: currently supporting only UChimeraAsset(s), but we should generalize for UMultiAnimAsset(s)
-		if (const UE::PoseSearch::FSearchIndexAsset* SearchIndexAsset = SearchResult.GetSearchIndexAsset())
+	if (const UE::PoseSearch::FSearchIndexAsset* SearchIndexAsset = SearchResult.GetSearchIndexAsset())
+	{
+		if (const FPoseSearchDatabaseAnimationAssetBase* DatabaseAnimationAssetBase = SearchResult.Database->GetDatabaseAnimationAsset<FPoseSearchDatabaseAnimationAssetBase>(*SearchIndexAsset))
 		{
-			if (const FPoseSearchDatabaseMultiAnimAsset* DatabaseMultiAnimAsset = SearchResult.Database->GetDatabaseAnimationAsset<FPoseSearchDatabaseMultiAnimAsset>(*SearchIndexAsset))
+			if (const UMultiAnimAsset* MultiAnimAsset = Cast<UMultiAnimAsset>(DatabaseAnimationAssetBase->GetAnimationAsset()))
 			{
-				if (const UChimeraAsset* ChimeraAsset = Cast<UChimeraAsset>(DatabaseMultiAnimAsset->GetAnimationAsset()))
+				check(MultiAnimAsset->GetNumRoles() == SearchContext.AnimInstances.Num());
+				check(MultiAnimAsset->GetNumRoles() == SearchContext.HistoryCollectors.Num());
+				check(MultiAnimAsset->GetNumRoles() == SearchContext.Roles.Num());
+
+				UE::PoseSearch::FRoleToIndex SearchContextRoleToIndex;
+				SearchContextRoleToIndex.Reserve(SearchContext.Roles.Num());
+				for (int32 SearchContextRoleIndex = 0; SearchContextRoleIndex < SearchContext.Roles.Num(); ++SearchContextRoleIndex)
 				{
-					check(ChimeraAsset->GetNumRoles() == SearchContext.AnimInstances.Num());
-					check(ChimeraAsset->GetNumRoles() == SearchContext.HistoryCollectors.Num());
-					check(ChimeraAsset->GetNumRoles() == SearchContext.Roles.Num());
+					SearchContextRoleToIndex.Add(SearchContext.Roles[SearchContextRoleIndex]) = SearchContextRoleIndex;
+				}
 
-					UE::PoseSearch::FRoleToIndex SearchContextRoleToIndex;
-					SearchContextRoleToIndex.Reserve(SearchContext.Roles.Num());
-					for (int32 SearchContextRoleIndex = 0; SearchContextRoleIndex < SearchContext.Roles.Num(); ++SearchContextRoleIndex)
-					{
-						SearchContextRoleToIndex.Add(SearchContext.Roles[SearchContextRoleIndex]) = SearchContextRoleIndex;
-					}
+				// mapping MultiAnimAsset Roles to SearchContext.* indexes
+				TArray<FTransform, TInlineAllocator<UE::PoseSearch::PreallocatedRolesNum>> ActorRootBoneTransforms;
+				const UWorld* DebugDrawWorld = nullptr;
+				for (int32 MultiAnimAssetRoleIndex = 0; MultiAnimAssetRoleIndex < MultiAnimAsset->GetNumRoles(); ++MultiAnimAssetRoleIndex)
+				{
+					const UE::PoseSearch::FRole& ChimeraAssetRole = MultiAnimAsset->GetRole(MultiAnimAssetRoleIndex);
+					const int32 SearchContextIndex = SearchContextRoleToIndex[ChimeraAssetRole];
 
-					// mapping ChimeraAsset Roles to SearchContext.* indexes
-					TArray<FTransform, TInlineAllocator<UE::PoseSearch::PreallocatedRolesNum>> ActorRootBoneTransforms;
-					const UWorld* DebugDrawWorld = nullptr;
-					for (int32 ChimeraAssetRoleIndex = 0; ChimeraAssetRoleIndex < ChimeraAsset->GetNumRoles(); ++ChimeraAssetRoleIndex)
+					if (UAnimInstance* AnimInstance = SearchContext.AnimInstances[SearchContextIndex].Get())
 					{
-						const UE::PoseSearch::FRole& ChimeraAssetRole = ChimeraAsset->GetRole(ChimeraAssetRoleIndex);
-						const int32 SearchContextIndex = SearchContextRoleToIndex[ChimeraAssetRole];
-
-						if (UAnimInstance* AnimInstance = SearchContext.AnimInstances[SearchContextIndex].Get())
-						{
-							const FTransform RootBoneTransform = AnimInstance->GetSkelMeshComponent()->GetBoneTransform(0);
-							ActorRootBoneTransforms.Add(RootBoneTransform);
+						const FTransform RootBoneTransform = AnimInstance->GetSkelMeshComponent()->GetBoneTransform(0);
+						ActorRootBoneTransforms.Add(RootBoneTransform);
 					
-							if (!DebugDrawWorld)
-							{
-								DebugDrawWorld = AnimInstance->GetWorld();
-							}
+						if (!DebugDrawWorld)
+						{
+							DebugDrawWorld = AnimInstance->GetWorld();
 						}
 					}
+				}
 
-					if (ActorRootBoneTransforms.Num() == ChimeraAsset->GetNumRoles())
-					{
-						// FullAlignedActorRootBoneTransforms is mapped to the ChimeraAsset roles:
-						// FullAlignedActorRootBoneTransforms[0] is for ChimeraAsset->GetRole(0)
-						TArray<FTransform, TInlineAllocator<UE::PoseSearch::PreallocatedRolesNum>> FullAlignedActorRootBoneTransforms;
-						FullAlignedActorRootBoneTransforms.SetNum(ChimeraAsset->GetNumRoles());
+				if (ActorRootBoneTransforms.Num() == MultiAnimAsset->GetNumRoles())
+				{
+					// FullAlignedActorRootBoneTransforms is mapped to the MultiAnimAsset roles:
+					// FullAlignedActorRootBoneTransforms[0] is for MultiAnimAsset->GetRole(0)
+					TArray<FTransform, TInlineAllocator<UE::PoseSearch::PreallocatedRolesNum>> FullAlignedActorRootBoneTransforms;
+					FullAlignedActorRootBoneTransforms.SetNum(MultiAnimAsset->GetNumRoles());
 
-						ChimeraAsset->CalculateWarpTransforms(SearchResult.AssetTime, ActorRootBoneTransforms, FullAlignedActorRootBoneTransforms, DebugDrawWorld);
-						ChimeraSearchResult.FullAlignedActorRootBoneTransforms.SetNum(ChimeraAsset->GetNumRoles());
+					// @todo: should it be SearchResult.AssetTime + DeltaTime?
+					MultiAnimAsset->CalculateWarpTransforms(SearchResult.AssetTime, ActorRootBoneTransforms, FullAlignedActorRootBoneTransforms, DebugDrawWorld);
+					ChimeraSearchResult.FullAlignedActorRootBoneTransforms.SetNum(MultiAnimAsset->GetNumRoles());
 
-						for (int32 ChimeraAssetRoleIndex = 0; ChimeraAssetRoleIndex < ChimeraAsset->GetNumRoles(); ++ChimeraAssetRoleIndex)
-						{
-							const UE::PoseSearch::FRole& ChimeraAssetRole = ChimeraAsset->GetRole(ChimeraAssetRoleIndex);
-							const int32 SearchContextIndex = SearchContextRoleToIndex[ChimeraAssetRole];
+					for (int32 ChimeraAssetRoleIndex = 0; ChimeraAssetRoleIndex < MultiAnimAsset->GetNumRoles(); ++ChimeraAssetRoleIndex)
+					{
+						const UE::PoseSearch::FRole& ChimeraAssetRole = MultiAnimAsset->GetRole(ChimeraAssetRoleIndex);
+						const int32 SearchContextIndex = SearchContextRoleToIndex[ChimeraAssetRole];
 
-							ChimeraSearchResult.FullAlignedActorRootBoneTransforms[SearchContextIndex] = FullAlignedActorRootBoneTransforms[ChimeraAssetRoleIndex];
-						}
+						ChimeraSearchResult.FullAlignedActorRootBoneTransforms[SearchContextIndex] = FullAlignedActorRootBoneTransforms[ChimeraAssetRoleIndex];
 					}
 				}
 			}
-		}
+			else
+			{
+				// support for non UMultiAnimAsset to be backward compatible with regular motion matching searches
+				check(SearchContext.AnimInstances.Num() == 1);
 
-		return ChimeraSearchResult;
+				if (UAnimInstance* AnimInstance = SearchContext.AnimInstances[0].Get())
+				{
+					const FTransform RootBoneTransform = AnimInstance->GetSkelMeshComponent()->GetBoneTransform(0);
+					ChimeraSearchResult.FullAlignedActorRootBoneTransforms.SetNum(1);
+					// @todo: should it be RootBoneTransform + root motion transform?
+					ChimeraSearchResult.FullAlignedActorRootBoneTransforms[0] = RootBoneTransform;
+				}
+			}
+		}
 	}
 
-	bool FSearchContext::IsValid() const
+	return ChimeraSearchResult;
+}
+
+static void InitSearchResults(TArray<FSearchResult>& SearchResults, TConstArrayView<UE::PoseSearch::FSearchResult> PoseSearchResults, TConstArrayView<FSearchContext> SearchContexts)
+{
+	// WIP!
+	// @todo: figure out multiple policies to use the most characters? right now only the best search is "valid" with the most characters
+	int32 BestSearchIndex = INDEX_NONE;
+	for (int32 SearchIndex = 0; SearchIndex < PoseSearchResults.Num(); ++SearchIndex)
 	{
-		if (Database == nullptr)
+		if (PoseSearchResults[SearchIndex].IsValid())
 		{
-			return false;
+			if (BestSearchIndex == INDEX_NONE)
+			{
+				BestSearchIndex = SearchIndex;
+			}
+			else if (SearchContexts[SearchIndex].Roles.Num() > SearchContexts[BestSearchIndex].Roles.Num())
+			{
+				BestSearchIndex = SearchIndex;
+			}
+			else if (SearchContexts[SearchIndex].Roles.Num() == SearchContexts[BestSearchIndex].Roles.Num() &&
+				PoseSearchResults[SearchIndex].PoseCost < PoseSearchResults[BestSearchIndex].PoseCost)
+			{
+				BestSearchIndex = SearchIndex;
+			}
 		}
+	}
 
-		const int32 Num = AnimInstances.Num();
-		if (Num < 1)
-		{
-			return false;
-		}
+	if (BestSearchIndex != INDEX_NONE)
+	{
+		SearchResults.SetNum(1);
+		SearchResults[0] = InitSearchResult(PoseSearchResults[BestSearchIndex], BestSearchIndex, SearchContexts[BestSearchIndex]);
+	}
+}
 
-		if (Num != HistoryCollectors.Num())
-		{
-			return false;
-		}
+// FSearchResult
+///////////////////////////////////////////////////////////
+bool FSearchContext::IsValid() const
+{
+	if (Database == nullptr)
+	{
+		return false;
+	}
+
+	const int32 Num = AnimInstances.Num();
+	if (Num < 1)
+	{
+		return false;
+	}
+
+	if (Num != HistoryCollectors.Num())
+	{
+		return false;
+	}
 
-		if (Num != Roles.Num())
+	if (Num != Roles.Num())
+	{
+		return false;
+	}
+
+	for (int32 IndexA = 0; IndexA < Num; ++IndexA)
+	{
+		if (AnimInstances[IndexA] == nullptr)
 		{
 			return false;
 		}
 
-		for (int32 IndexA = 0; IndexA < Num; ++IndexA)
+		for (int32 IndexB = IndexA + 1; IndexB < Num; ++IndexB)
 		{
-			if (AnimInstances[IndexA] == nullptr)
+			if (AnimInstances[IndexA] == AnimInstances[IndexB])
 			{
 				return false;
 			}
+		}
+	}
 
-			for (int32 IndexB = IndexA + 1; IndexB < Num; ++IndexB)
+	for (int32 IndexA = 0; IndexA < Num; ++IndexA)
+	{
+		for (int32 IndexB = IndexA + 1; IndexB < Num; ++IndexB)
+		{
+			if (Roles[IndexA] == Roles[IndexB])
 			{
-				if (AnimInstances[IndexA] == AnimInstances[IndexB])
-				{
-					return false;
-				}
+				return false;
 			}
 		}
+	}
 
-		for (int32 IndexA = 0; IndexA < Num; ++IndexA)
+	for (int32 IndexA = 0; IndexA < Num; ++IndexA)
+	{
+		if (HistoryCollectors[IndexA] == nullptr)
 		{
-			for (int32 IndexB = IndexA + 1; IndexB < Num; ++IndexB)
-			{
-				if (Roles[IndexA] == Roles[IndexB])
-				{
-					return false;
-				}
-			}
+			return false;
 		}
 
-		for (int32 IndexA = 0; IndexA < Num; ++IndexA)
+		for (int32 IndexB = IndexA + 1; IndexB < Num; ++IndexB)
 		{
-			if (HistoryCollectors[IndexA] == nullptr)
+			if (HistoryCollectors[IndexA] == HistoryCollectors[IndexB])
 			{
 				return false;
 			}
-
-			for (int32 IndexB = IndexA + 1; IndexB < Num; ++IndexB)
-			{
-				if (HistoryCollectors[IndexA] == HistoryCollectors[IndexB])
-				{
-					return false;
-				}
-			}
 		}
-
-		return true;
 	}
 
-	bool FSearchContext::IsEquivalent(const FSearchContext& Other) const
+	return true;
+}
+
+bool FSearchContext::IsEquivalent(const FSearchContext& Other) const
+{
+	if (Database != Other.Database)
 	{
-		if (Database != Other.Database)
-		{
-			return false;
-		}
+		return false;
+	}
 
-		const int32 Num = AnimInstances.Num();
-		if (Num != Other.AnimInstances.Num())
-		{
-			return false;
-		}
+	const int32 Num = AnimInstances.Num();
+	if (Num != Other.AnimInstances.Num())
+	{
+		return false;
+	}
 
-		int32 CommonRoledAnimInstances = 0;
-		for (int32 IndexThis = 0; IndexThis < Num; ++IndexThis)
+	int32 CommonRoledAnimInstances = 0;
+	for (int32 IndexThis = 0; IndexThis < Num; ++IndexThis)
+	{
+		for (int32 IndexOther = 0; IndexOther < Num; ++IndexOther)
 		{
-			for (int32 IndexOther = 0; IndexOther < Num; ++IndexOther)
+			if (AnimInstances[IndexThis] == Other.AnimInstances[IndexOther] &&
+				HistoryCollectors[IndexThis] == Other.HistoryCollectors[IndexOther] &&
+				Roles[IndexThis] == Other.Roles[IndexOther])
 			{
-				if (AnimInstances[IndexThis] == Other.AnimInstances[IndexOther] &&
-					HistoryCollectors[IndexThis] == Other.HistoryCollectors[IndexOther] &&
-					Roles[IndexThis] == Other.Roles[IndexOther])
-				{
-					++CommonRoledAnimInstances;
-				}
+				++CommonRoledAnimInstances;
 			}
 		}
+	}
 
-		// using >= in case there are duplicated animinstances in this or Other (this->IsValid() should be false!)
-		if (CommonRoledAnimInstances >= Num)
-		{
-			return true;
-		}
-
-		return false;
+	// using >= in case there are duplicated animinstances in this or Other (this->IsValid() should be false!)
+	if (CommonRoledAnimInstances >= Num)
+	{
+		return true;
 	}
+
+	return false;
 }
 
-void UChimeraIslandComponent::TickComponent(float DeltaTime, enum ELevelTick TickType, FActorComponentTickFunction* ThisTickFunction)
+// FSearchResult
+///////////////////////////////////////////////////////////
+bool FSearchResult::operator==(const FSearchResult& Other) const
 {
-	// Called before any skeletal mesh component tick, when there aren't animation jobs flying. No need to FScopeLock Lock(&Mutex);
-	using namespace UE::Chimera;
+	return static_cast<const UE::PoseSearch::FSearchResult&>(*this) == static_cast<const UE::PoseSearch::FSearchResult&>(Other) &&
+				FullAlignedActorRootBoneTransforms.Num() == Other.FullAlignedActorRootBoneTransforms.Num() &&
+				0 == FMemory::Memcmp(FullAlignedActorRootBoneTransforms.GetData(), Other.FullAlignedActorRootBoneTransforms.GetData(), FullAlignedActorRootBoneTransforms.Num() * sizeof(FTransform));
+}
 
+// FIslandPreTickFunction
+///////////////////////////////////////////////////////////
+void FIsland::FPreTickFunction::ExecuteTick(float DeltaTime, enum ELevelTick TickType, ENamedThreads::Type CurrentThread, const FGraphEventRef& MyCompletionGraphEvent)
+{
+	// Called before any skeletal mesh component tick, when there aren't animation jobs flying. No need to FScopeLock Lock(&Mutex);
 	// generating trajectories before running any of the skeletal mesh component ticks
-	for (FSearchContext& SearchContext : SearchContexts)
+	for (const FSearchContext& SearchContext : Island->SearchContexts)
 	{
 		for (int32 Index = 0; Index < SearchContext.AnimInstances.Num(); ++Index)
 		{
-			if (UAnimInstance* AnimInstance = SearchContext.AnimInstances[Index].Get())
+			if (const UAnimInstance* AnimInstance = SearchContext.AnimInstances[Index].Get())
 			{
-				// since UChimeraIslandComponent has a tick dependency with the USkeletalMeshComponent it's safe modify the FAnimNode_PoseSearchHistoryCollector_Base
+				// since FIsland has a tick dependency with the USkeletalMeshComponent it's safe modify the FAnimNode_PoseSearchHistoryCollector_Base
 				FAnimNode_PoseSearchHistoryCollector_Base* HistoryCollector = const_cast<FAnimNode_PoseSearchHistoryCollector_Base*>(SearchContext.HistoryCollectors[Index]);
 				check(HistoryCollector);
 				HistoryCollector->GenerateTrajectory(AnimInstance);
@@ -216,12 +274,43 @@ void UChimeraIslandComponent::TickComponent(float DeltaTime, enum ELevelTick Tic
 	}
 }
 
-void UChimeraIslandComponent::DebugDraw(const FColor& Color) const
+// FPostTickFunction
+///////////////////////////////////////////////////////////
+void FIsland::FPostTickFunction::ExecuteTick(float DeltaTime, enum ELevelTick TickType, ENamedThreads::Type CurrentThread, const FGraphEventRef& MyCompletionGraphEvent)
+{
+	// do nothing
+}
+
+// FIsland
+///////////////////////////////////////////////////////////
+FIsland::FIsland(ULevel* Level)
+{
+	PreTickFunction.bCanEverTick = true;
+	PreTickFunction.bStartWithTickEnabled = true;
+	PreTickFunction.TickGroup = ETickingGroup::TG_PrePhysics;
+	PreTickFunction.Island = this;
+	PreTickFunction.SetTickFunctionEnable(true);
+	PreTickFunction.RegisterTickFunction(Level);
+
+	PostTickFunction.bCanEverTick = true;
+	PostTickFunction.bStartWithTickEnabled = true;
+	PostTickFunction.TickGroup = ETickingGroup::TG_PrePhysics;
+	PostTickFunction.SetTickFunctionEnable(true);
+	PostTickFunction.RegisterTickFunction(Level);
+}
+
+FIsland::~FIsland()
+{
+	PreTickFunction.UnRegisterTickFunction();
+	PostTickFunction.UnRegisterTickFunction();
+
+	Uninject();
+}
+
+void FIsland::DebugDraw(const FColor& Color) const
 {
 	// called only by UChimeraSubsystem::Tick so no need to lock SearchResultsMutex to protect the read of SearchContexts
 #if ENABLE_DRAW_DEBUG
-	using namespace UE::Chimera;
-
 	check(IsInGameThread());
 
 	if (CVarChimeraShowIslands.GetValueOnAnyThread())
@@ -234,7 +323,7 @@ void UChimeraIslandComponent::DebugDraw(const FColor& Color) const
 				{
 					const FVector Position = AnimInstance->GetSkelMeshComponent()->GetComponentLocation();
 					const float BroadPhaseRadius = SearchContext.BroadPhaseRadiuses[Index];
-					DrawDebugCircle(GetWorld(), Position, BroadPhaseRadius, 40, Color, false, 0.f, SDPG_Foreground, 0.f, FVector::XAxisVector, FVector::YAxisVector, false);
+					DrawDebugCircle(AnimInstance->GetWorld(), Position, BroadPhaseRadius, 40, Color, false, 0.f, SDPG_Foreground, 0.f, FVector::XAxisVector, FVector::YAxisVector, false);
 				}
 			}
 		}
@@ -242,7 +331,7 @@ void UChimeraIslandComponent::DebugDraw(const FColor& Color) const
 #endif // ENABLE_DRAW_DEBUG
 }
 
-void UChimeraIslandComponent::InjectToActor(AActor* Actor)
+void FIsland::InjectToActor(AActor* Actor)
 {
 	check(IsInGameThread());
 
@@ -253,24 +342,52 @@ void UChimeraIslandComponent::InjectToActor(AActor* Actor)
 		{
 			if (USkeletalMeshComponent* SkeletalMeshComponent = Actor->GetComponentByClass<USkeletalMeshComponent>())
 			{
-				// tick order: CharacterMovementComponent(s) -> ChimeraIslandComponent -> SkeletalMeshComponent(s)
+				const bool bIsFirstInjectedActor = IsUninjected();
+
+				//	tick order: 
+				//		CharacterMovementComponent(s) ->
+				//			Island.PreTickFunction ->
+				//				first injected actor SkeletalMeshComponent ->
+				//					Island.PostTickFunction ->
+				//						other SkeletalMeshComponent(s)
 				CharacterMovementComponents.AddUnique(CharacterMovementComponent);
 				SkeletalMeshComponents.AddUnique(SkeletalMeshComponent);
+	
+				// making sure that if we add a unique CharacterMovementComponent, we add as well a unique SkeletalMeshComponent
+				// (so we can remove them later on in a consistent fashion)
+				check(CharacterMovementComponents.Num() == SkeletalMeshComponents.Num());
+
+				PreTickFunction.AddPrerequisite(Actor, CharacterMovementComponent->PrimaryComponentTick);
+				SkeletalMeshComponent->PrimaryComponentTick.AddPrerequisite(Actor, PreTickFunction);
 
-				AddTickPrerequisiteComponent(CharacterMovementComponent);
-				SkeletalMeshComponent->AddTickPrerequisiteComponent(this);
+				if (bIsFirstInjectedActor)
+				{
+					PostTickFunction.AddPrerequisite(Actor, SkeletalMeshComponent->PrimaryComponentTick);
+				}
+				else 
+				{
+					SkeletalMeshComponent->PrimaryComponentTick.AddPrerequisite(Actor, PostTickFunction);
+				}
+			}
+			else
+			{
+				UE_LOG(LogChimera, Error, TEXT("FIsland::InjectToActor requires AActor %s to have a USkeletalMeshComponent to work!"), *Actor->GetName());
 			}
 		}
+		else
+		{
+			UE_LOG(LogChimera, Error, TEXT("FIsland::InjectToActor requires AActor %s to have a UCharacterMovementComponent to work!"), *Actor->GetName());
+		}
 	}
 }
 
-void UChimeraIslandComponent::AddSearchContext(const UE::Chimera::FSearchContext& SearchContext)
+void FIsland::AddSearchContext(const FSearchContext& SearchContext)
 {
 #if DO_CHECK
 	check(SearchContext.IsValid());
 	check(IsInGameThread());
 
-	for (const UE::Chimera::FSearchContext& ContainedSearchContext : SearchContexts)
+	for (const FSearchContext& ContainedSearchContext : SearchContexts)
 	{
 		check(!ContainedSearchContext.IsEquivalent(SearchContext));
 	}
@@ -278,19 +395,42 @@ void UChimeraIslandComponent::AddSearchContext(const UE::Chimera::FSearchContext
 	SearchContexts.Add(SearchContext);
 }
 
-void UChimeraIslandComponent::Uninject()
+void FIsland::Uninject()
 {
+	// Called by UChimeraSubsystem::Tick when there aren't animation jobs flying. No need to FScopeLock Lock(&Mutex);
 	check(IsInGameThread());
 
-	// Called by UChimeraSubsystem::Tick when there aren't animation jobs flying. No need to FScopeLock Lock(&Mutex);
-	for (TWeakObjectPtr<UCharacterMovementComponent>& CharacterMovementComponentPtr : CharacterMovementComponents)
-	{
-		RemoveTickPrerequisiteComponent(CharacterMovementComponentPtr.Get());
-	}
+	check(CharacterMovementComponents.Num() == SkeletalMeshComponents.Num());
 
-	for (TWeakObjectPtr<USkeletalMeshComponent>& SkeletalMeshComponentPtr : SkeletalMeshComponents)
+	for (int32 ActorIndex = 0; ActorIndex < CharacterMovementComponents.Num(); ++ActorIndex)
 	{
-		SkeletalMeshComponentPtr->RemoveTickPrerequisiteComponent(this);
+		UCharacterMovementComponent* CharacterMovementComponent = CharacterMovementComponents[ActorIndex].Get();
+		USkeletalMeshComponent* SkeletalMeshComponent = SkeletalMeshComponents[ActorIndex].Get();
+
+		if (CharacterMovementComponent)
+		{
+			check(SkeletalMeshComponent);
+
+			AActor* Actor = CharacterMovementComponent->GetOwner();
+			check(Actor);
+
+			PreTickFunction.RemovePrerequisite(Actor, CharacterMovementComponent->PrimaryComponentTick);
+			SkeletalMeshComponent->PrimaryComponentTick.RemovePrerequisite(Actor, PreTickFunction);
+
+			const bool bIsFirstInjectedActor = ActorIndex == 0;
+			if (bIsFirstInjectedActor)
+			{
+				PostTickFunction.RemovePrerequisite(Actor, SkeletalMeshComponent->PrimaryComponentTick);
+			}
+			else
+			{
+				SkeletalMeshComponent->PrimaryComponentTick.RemovePrerequisite(Actor, PostTickFunction);
+			}
+		}
+		else
+		{
+			check(!SkeletalMeshComponent);
+		}
 	}
 
 	CharacterMovementComponents.Reset();
@@ -301,15 +441,13 @@ void UChimeraIslandComponent::Uninject()
 	bSearchPerfomed = false;
 }
 
-bool UChimeraIslandComponent::IsUninjected()
+bool FIsland::IsUninjected() const
 {
 	return SkeletalMeshComponents.IsEmpty();
 }
 
-bool UChimeraIslandComponent::DoSearch_AnyThread(UObject* AnimInstance, FChimeraBlueprintResult& Result)
+bool FIsland::DoSearch_AnyThread(UObject* AnimInstance, const FPoseSearchContinuingProperties& ContinuingProperties, FChimeraBlueprintResult& Result)
 {
-	using namespace UE::Chimera;
-
 	bool bDoPerfomSearch = false;
 	{
 		// thread safety note!
@@ -317,8 +455,8 @@ bool UChimeraIslandComponent::DoSearch_AnyThread(UObject* AnimInstance, FChimera
 		// why:		UPoseSearchLibrary::MotionMatch could call via AnimInstance GetProxyOnAnyThread<FAnimInstanceProxy>() that, if on GameThread, 
 		//			could call UAnimInstance::HandleExistingParallelEvaluationTask 
 		// fix:		avoid UPoseSearchLibrary::MotionMatch calls wrapped by any lock, at the cost of eventually (by design should be NEVER) performing the searches twice
-		//			By design we should inject ticks dependencies (added by UChimeraIslandComponent::InjectToActor via AddTickPrerequisiteComponent), so concurrently 
-		//			fly of UChimeraIslandComponent within the same island that requires searches is forbidden
+		//			By design we should inject ticks dependencies (added by FIsland::InjectToActor via AddTickPrerequisiteComponent), so concurrently 
+		//			fly of FIsland within the same island that requires searches is forbidden
 		FScopeLock Lock(&SearchResultsMutex);
 		bDoPerfomSearch = !bSearchPerfomed;
 	}
@@ -340,13 +478,13 @@ bool UChimeraIslandComponent::DoSearch_AnyThread(UObject* AnimInstance, FChimera
 			const UPoseSearchDatabase* Database = SearchContext.Database.Get();
 			if (!Database)
 			{
-				UE_LOG(LogChimera, Error, TEXT("UChimeraIslandComponent::DoSearch_AnyThread invalid context database"));
+				UE_LOG(LogChimera, Error, TEXT("FIsland::DoSearch_AnyThread invalid context database"));
 				return false;
 			}
 
 			if (!Database->Schema)
 			{
-				UE_LOG(LogChimera, Error, TEXT("UChimeraIslandComponent::DoSearch_AnyThread invalid schema for context database %s"), *Database->GetName());
+				UE_LOG(LogChimera, Error, TEXT("FIsland::DoSearch_AnyThread invalid schema for context database %s"), *Database->GetName());
 				return false;
 			}
 
@@ -356,7 +494,7 @@ bool UChimeraIslandComponent::DoSearch_AnyThread(UObject* AnimInstance, FChimera
 				UAnimInstance* SearchContextAnimInstance = AnimInstancePtr.Get();
 				if (!SearchContextAnimInstance)
 				{
-					UE_LOG(LogChimera, Error, TEXT("UChimeraIslandComponent::DoSearch_AnyThread null anim instance"));
+					UE_LOG(LogChimera, Error, TEXT("FIsland::DoSearch_AnyThread null anim instance"));
 					return false;
 				}
 
@@ -372,9 +510,41 @@ bool UChimeraIslandComponent::DoSearch_AnyThread(UObject* AnimInstance, FChimera
 			const UObject* AssetsToSearch[] = { Database };
 			FPoseSearchFutureProperties PoseSearchFutureProperties;
 
+			FPoseSearchContinuingProperties ContinuingPropertiesToUse = SearchContext.ContinuingProperties;
+
+			// @todo: implemennt this carefully, since it'll cause thread safety concerns depending on which actors is calling the query first 
+			// and if their animations get integrated by different play rates
+			
+			//if (ContinuingProperties.PlayingAsset)
+			//{
+			//	if (const UMultiAnimAsset* MultiAnimAsset = Cast<UMultiAnimAsset>(ContinuingProperties.PlayingAsset))
+			//	{
+			//		// @todo: if they are compatible... etcetc MultiAnimAsset->GetNumRoles
+			//		ContinuingPropertiesToUse = ContinuingProperties;
+			//	}
+			//	else if (const UMultiAnimAsset* SearchContextMultiAnimAsset = Cast<UMultiAnimAsset>(SearchContext.ContinuingProperties.PlayingAsset))
+			//	{
+			//		for (int32 RoleIndex = 0; RoleIndex < SearchContextMultiAnimAsset->GetNumRoles(); ++RoleIndex)
+			//		{
+			//			if (SearchContextMultiAnimAsset->GetAnimationAsset(SearchContextMultiAnimAsset->GetRole(RoleIndex)) == ContinuingProperties.PlayingAsset)
+			//			{
+			//				ContinuingPropertiesToUse.PlayingAssetAccumulatedTime = ContinuingProperties.PlayingAssetAccumulatedTime;
+			//				break;
+			//			}
+			//		}
+			//	}
+			//	else if (!SearchContext.ContinuingProperties.PlayingAsset)
+			//	{
+			//		// @todo: should we scan all the database iterating over the UMultiAnimAsset to look for ContinuingProperties.PlayingAsset?
+			//	}
+			//	else
+			//	{
+			//	}
+			//}
+
 			// @todo: we could perform multiple UPoseSearchLibrary::MotionMatch in parallel!
 			const UE::PoseSearch::FSearchResult PoseSearchResult = UPoseSearchLibrary::MotionMatch(AnimInstances, SearchContext.Roles,
-				PoseHistories, AssetsToSearch, SearchContext.ContinuingProperties, PoseSearchFutureProperties);
+				PoseHistories, AssetsToSearch, ContinuingPropertiesToUse, PoseSearchFutureProperties);
 
 			if (PoseSearchResult.PoseCost.GetTotalCost() < SearchContext.MaxCost)
 			{
@@ -382,46 +552,31 @@ bool UChimeraIslandComponent::DoSearch_AnyThread(UObject* AnimInstance, FChimera
 			}
 		}
 
-		// making sure we called Uninject()
-		check(SearchResults.IsEmpty());
-
-		////////////////////////////////////////////////////////////////////////////////////////////////////////////////
-		// WIP!
-		// @todo: figure out multiple policies to use the most characters? right now only the best search is "valid" with the most characters
 		if (!PoseSearchResults.IsEmpty())
 		{
 			// locking to update SearchResults and bSearchPerfomed
 			FScopeLock Lock(&SearchResultsMutex);
 
-			int32 BestSearchIndex = INDEX_NONE;
-			for (int32 SearchIndex = 0; SearchIndex < PoseSearchResults.Num(); ++SearchIndex)
+			if (!bSearchPerfomed)
 			{
-				if (PoseSearchResults[SearchIndex].IsValid())
-				{
-					if (BestSearchIndex == INDEX_NONE)
-					{
-						BestSearchIndex = SearchIndex;
-					}
-					else if (SearchContexts[SearchIndex].Roles.Num() > SearchContexts[BestSearchIndex].Roles.Num())
-					{
-						BestSearchIndex = SearchIndex;
-					}
-					else if (SearchContexts[SearchIndex].Roles.Num() == SearchContexts[BestSearchIndex].Roles.Num() &&
-						PoseSearchResults[SearchIndex].PoseCost < PoseSearchResults[BestSearchIndex].PoseCost)
-					{
-						BestSearchIndex = SearchIndex;
-					}
-				}
-			}
+				InitSearchResults(SearchResults, PoseSearchResults, SearchContexts);
 
-			if (BestSearchIndex != INDEX_NONE)
-			{
-				SearchResults.SetNum(1);
-				SearchResults[0] = InitSearchResult(PoseSearchResults[BestSearchIndex], BestSearchIndex, SearchContexts[BestSearchIndex]);
+				bSearchPerfomed = true;
 			}
-			////////////////////////////////////////////////////////////////////////////////////////////////////////////////
+			else
+			{
+				UE_LOG(LogChimera, Warning, TEXT("FIsland::DoSearch_AnyThread performance warning: performed duplicated search :("));
+				
+#if DO_CHECK
+				TArray<FSearchResult> CompareSearchResults;
+				InitSearchResults(CompareSearchResults, PoseSearchResults, SearchContexts);
 
-			bSearchPerfomed = true;
+				if (CompareSearchResults != SearchResults)
+				{
+					UE_LOG(LogChimera, Error, TEXT("FIsland::DoSearch_AnyThread duplicated search differs from original one. Searches are NOT deterministic!"));
+				}
+#endif // DO_CHECK
+			}
 	
 			// calling this funtion within this scope since we already locked SearchResultsMutex
 			return GetResult_AnyThread(AnimInstance, Result);
@@ -431,10 +586,8 @@ bool UChimeraIslandComponent::DoSearch_AnyThread(UObject* AnimInstance, FChimera
 	return GetResult_AnyThread(AnimInstance, Result);
 }
 
-bool UChimeraIslandComponent::GetResult_AnyThread(UObject* AnimInstance, FChimeraBlueprintResult& Result)
+bool FIsland::GetResult_AnyThread(UObject* AnimInstance, FChimeraBlueprintResult& Result)
 {
-	using namespace UE::Chimera;
-
 	// locking to read SearchResults
 	FScopeLock Lock(&SearchResultsMutex);
 
@@ -499,13 +652,11 @@ bool UChimeraIslandComponent::GetResult_AnyThread(UObject* AnimInstance, FChimer
 	return false;
 }
 
-const UE::Chimera::FSearchResult* UChimeraIslandComponent::FindSearchResult(const UE::Chimera::FSearchContext& SearchContext) const
+const FSearchResult* FIsland::FindSearchResult(const FSearchContext& SearchContext) const
 {
 	// called only by UChimeraSubsystem::Tick via UChimeraSubsystem::PopulateContinuingProperties so no need to lock SearchResultsMutex to protect the read of SearchResults
 	check(IsInGameThread());
 
-	using namespace UE::Chimera;
-
 	// searching for InSearchContext in all the SearchContexts referenced by valid active SearchResults
 	for (const FSearchResult& SearchResult : SearchResults)
 	{
@@ -519,3 +670,5 @@ const UE::Chimera::FSearchResult* UChimeraIslandComponent::FindSearchResult(cons
 	}
 	return nullptr;
 }
+
+} // namespace UE::Chimera
diff --git a/Engine/Plugins/Experimental/Animation/Chimera/Source/Runtime/Private/ChimeraLibrary.cpp b/Engine/Plugins/Experimental/Animation/Chimera/Source/Runtime/Private/ChimeraLibrary.cpp
--- a/Engine/Plugins/Experimental/Animation/Chimera/Source/Runtime/Private/ChimeraLibrary.cpp
+++ b/Engine/Plugins/Experimental/Animation/Chimera/Source/Runtime/Private/ChimeraLibrary.cpp
@@ -1,30 +1,42 @@
 // Copyright Epic Games, Inc. All Rights Reserved.
 
 #include "Chimera/ChimeraLibrary.h"
+#include "Animation/AnimMontage.h"
 #include "Chimera/ChimeraSubsystem.h"
 #include "Engine/World.h"
 
-FChimeraBlueprintResult UChimeraLibrary::ChimeraQuery_Pure(TArray<FChimeraAvailability> Availabilities, UObject* AnimInstance, FName PoseHistoryName, bool bValidateResultAgainstAvailabilities)
+FChimeraBlueprintResult UChimeraLibrary::ChimeraQuery_Pure(TArray<FChimeraAvailability> Availabilities, UObject* AnimInstance, FPoseSearchContinuingProperties ContinuingProperties, FName PoseHistoryName, bool bValidateResultAgainstAvailabilities)
 {
 	FChimeraBlueprintResult Result;
 	if (UChimeraSubsystem* ChimeraSubsystem = UChimeraSubsystem::GetSubsystem_AnyThread(AnimInstance))
 	{
-		ChimeraSubsystem->Query_AnyThread(Availabilities, AnimInstance, Result, PoseHistoryName, nullptr, bValidateResultAgainstAvailabilities);
+		ChimeraSubsystem->Query_AnyThread(Availabilities, AnimInstance, ContinuingProperties, Result, PoseHistoryName, nullptr, bValidateResultAgainstAvailabilities);
 	}
 	return Result;
 }
 
-FChimeraBlueprintResult UChimeraLibrary::ChimeraQuery(TArray<FChimeraAvailability> Availabilities, UObject* AnimInstance, FName PoseHistoryName, bool bValidateResultAgainstAvailabilities)
+FChimeraBlueprintResult UChimeraLibrary::ChimeraQuery(TArray<FChimeraAvailability> Availabilities, UObject* AnimInstance, FPoseSearchContinuingProperties ContinuingProperties, FName PoseHistoryName, bool bValidateResultAgainstAvailabilities)
 {
-	return ChimeraQuery_Pure(Availabilities, AnimInstance, PoseHistoryName, bValidateResultAgainstAvailabilities);
+	return ChimeraQuery_Pure(Availabilities, AnimInstance, ContinuingProperties, PoseHistoryName, bValidateResultAgainstAvailabilities);
 }
 
-FChimeraBlueprintResult UChimeraLibrary::ChimeraQuery(const TArrayView<const FChimeraAvailability> Availabilities, UObject* AnimInstance, const FAnimNode_PoseSearchHistoryCollector_Base* HistoryCollector, bool bValidateResultAgainstAvailabilities)
+FChimeraBlueprintResult UChimeraLibrary::ChimeraQuery(const TArrayView<const FChimeraAvailability> Availabilities, UObject* AnimInstance, const FPoseSearchContinuingProperties& ContinuingProperties, const FAnimNode_PoseSearchHistoryCollector_Base* HistoryCollector, bool bValidateResultAgainstAvailabilities)
 {
 	FChimeraBlueprintResult Result;
 	if (UChimeraSubsystem* ChimeraSubsystem = UChimeraSubsystem::GetSubsystem_AnyThread(AnimInstance))
 	{
-		ChimeraSubsystem->Query_AnyThread(Availabilities, AnimInstance, Result, FName(), HistoryCollector, bValidateResultAgainstAvailabilities);
+		ChimeraSubsystem->Query_AnyThread(Availabilities, AnimInstance, ContinuingProperties, Result, FName(), HistoryCollector, bValidateResultAgainstAvailabilities);
 	}
 	return Result;
 }
+
+FPoseSearchContinuingProperties UChimeraLibrary::GetMontageContinuingProperties(UAnimInstance* AnimInstance)
+{
+	FPoseSearchContinuingProperties ContinuingProperties;
+	if (const FAnimMontageInstance* AnimMontageInstance = AnimInstance->GetActiveMontageInstance())
+	{
+		ContinuingProperties.PlayingAsset = AnimMontageInstance->Montage;
+		ContinuingProperties.PlayingAssetAccumulatedTime = AnimMontageInstance->DeltaTimeRecord.GetPrevious();
+	}
+	return ContinuingProperties;
+}
\ No newline at end of file
diff --git a/Engine/Plugins/Experimental/Animation/Chimera/Source/Runtime/Private/ChimeraSubsystem.cpp b/Engine/Plugins/Experimental/Animation/Chimera/Source/Runtime/Private/ChimeraSubsystem.cpp
--- a/Engine/Plugins/Experimental/Animation/Chimera/Source/Runtime/Private/ChimeraSubsystem.cpp
+++ b/Engine/Plugins/Experimental/Animation/Chimera/Source/Runtime/Private/ChimeraSubsystem.cpp
@@ -230,41 +230,32 @@ namespace UE::Chimera
 
 FCriticalSection UChimeraSubsystem::RetrieveSubsystemMutex;
 
-UChimeraIslandComponent& UChimeraSubsystem::CreateIsland()
+UE::Chimera::FIsland& UChimeraSubsystem::CreateIsland()
 {
-	UChimeraIslandComponent& Island = *Islands.Add_GetRef(NewObject<UChimeraIslandComponent>()).Get();
-	Island.PrimaryComponentTick.bCanEverTick = true;
-	Island.PrimaryComponentTick.bStartWithTickEnabled = true;
-	Island.SetTickGroup(ETickingGroup::TG_PrePhysics);
-	Island.AddToRoot();
-	Island.RegisterComponentWithWorld(GetWorld());
-	return Island;
+	return *Islands.Add_GetRef(new UE::Chimera::FIsland(ToRawPtr(GetWorld()->PersistentLevel)));
 }
 
-UChimeraIslandComponent& UChimeraSubsystem::GetAvailableIsland()
+void UChimeraSubsystem::DestroyIsland(int32 Index)
+{
+	delete Islands[Index];
+	Islands.RemoveAt(Index);
+}
+
+UE::Chimera::FIsland& UChimeraSubsystem::GetAvailableIsland()
 {
-	for (TObjectPtr<UChimeraIslandComponent>& IslandPtr : Islands)
+	using namespace UE::Chimera;
+
+	for (FIsland* Island : Islands)
 	{
-		UChimeraIslandComponent& Island = *IslandPtr.Get();
-		if (Island.IsUninjected())
+		if (Island->IsUninjected())
 		{
-			return Island;
+			return *Island;
 		}
 	}
 
 	return CreateIsland();
 }
 
-void UChimeraSubsystem::DestroyIsland(int32 Index)
-{
-	UChimeraIslandComponent& Island = *Islands[Index].Get();
-	Island.Uninject();
-	Island.RemoveFromRoot();
-	Island.UnregisterComponent();
-	Island.DestroyComponent();
-	Islands.RemoveAt(Index);
-}
-
 void UChimeraSubsystem::DestroyAllIslands()
 {
 	for (int32 IslandIndex = Islands.Num() - 1; IslandIndex >= 0; --IslandIndex)
@@ -275,23 +266,26 @@ void UChimeraSubsystem::DestroyAllIslands()
 
 void UChimeraSubsystem::UninjectAllIslands()
 {
+	using namespace UE::Chimera;
+
 	check(IsInGameThread());
 
-	for (TObjectPtr<UChimeraIslandComponent>& IslandPtr : Islands)
+	for (FIsland* Island : Islands)
 	{
-		IslandPtr.Get()->Uninject();
+		Island->Uninject();
 	}
 }
 
 bool UChimeraSubsystem::ValidateAllIslands() const
 {
 #if DO_CHECK
+	using namespace UE::Chimera;
 
 	TSet<TWeakObjectPtr<UCharacterMovementComponent>> CharacterMovementComponents;
 	TSet<TWeakObjectPtr<USkeletalMeshComponent>> SkeletalMeshComponents;
 
 	bool bAlreadyInSet = false;
-	for (const TObjectPtr<UChimeraIslandComponent>& Island : Islands)
+	for (const FIsland* Island : Islands)
 	{
 		for (const TWeakObjectPtr<UCharacterMovementComponent>& CharacterMovementComponent : Island->GetCharacterMovementComponents())
 		{
@@ -317,12 +311,14 @@ bool UChimeraSubsystem::ValidateAllIslands() const
 
 void UChimeraSubsystem::PopulateContinuingProperties(UE::Chimera::FSearchContext& SearchContext, float DeltaSeconds) const
 {
+	using namespace UE::Chimera;
+
 	check(IsInGameThread());
 
 	// searching this SearchContext in all the islands to initialize its continuing pose
-	for (const UChimeraIslandComponent* Island : Islands)
+	for (const FIsland* Island : Islands)
 	{
-		if (const UE::Chimera::FSearchResult* SearchResult = Island->FindSearchResult(SearchContext))
+		if (const FSearchResult* SearchResult = Island->FindSearchResult(SearchContext))
 		{
 			// is still valid...
 			if (SearchResult->IsValid())
@@ -341,8 +337,10 @@ void UChimeraSubsystem::PopulateContinuingProperties(UE::Chimera::FSearchContext
 	}
 }
 
-UChimeraIslandComponent* UChimeraSubsystem::FindIsland(UObject* InAnimInstance)
+UE::Chimera::FIsland* UChimeraSubsystem::FindIsland(UObject* InAnimInstance)
 {
+	using namespace UE::Chimera;
+
 	if (InAnimInstance)
 	{
 		if (UAnimInstance* AnimInstance = Cast<UAnimInstance>(InAnimInstance))
@@ -351,13 +349,13 @@ UChimeraIslandComponent* UChimeraSubsystem::FindIsland(UObject* InAnimInstance)
 			{
 				if (USkeletalMeshComponent* SkeletalMeshComponent = Actor->GetComponentByClass<USkeletalMeshComponent>())
 				{
-					for (const TObjectPtr<UChimeraIslandComponent>& Island : Islands)
+					for (FIsland* Island : Islands)
 					{
 						for (const TWeakObjectPtr<USkeletalMeshComponent>& IslandSkeletalMeshComponent : Island->GetSkeletalMeshComponents())
 						{
 							if (IslandSkeletalMeshComponent.Get() == SkeletalMeshComponent)
 							{
-								return Island.Get();
+								return Island;
 							}
 						}
 					}
@@ -564,7 +562,7 @@ void UChimeraSubsystem::Tick(float DeltaSeconds)
 
 	for (const FSearchContextGroup& SearchContextGroup : SearchContextGroups)
 	{
-		UChimeraIslandComponent& Island = GetAvailableIsland();
+		UE::Chimera::FIsland& Island = GetAvailableIsland();
 
 		for (const UAnimInstance* AnimInstance : SearchContextGroup.AnimInstances)
 		{
@@ -587,7 +585,8 @@ TStatId UChimeraSubsystem::GetStatId() const
 	RETURN_QUICK_DECLARE_CYCLE_STAT(UChimeraSubsystem, STATGROUP_Tickables);
 }
 
-void UChimeraSubsystem::Query_AnyThread(const TArrayView<const FChimeraAvailability> Availabilities, UObject* AnimInstance, FChimeraBlueprintResult& Result,
+void UChimeraSubsystem::Query_AnyThread(const TArrayView<const FChimeraAvailability> Availabilities, UObject* AnimInstance, 
+	const FPoseSearchContinuingProperties& ContinuingProperties, FChimeraBlueprintResult& Result,
 	FName PoseHistoryName, const FAnimNode_PoseSearchHistoryCollector_Base* HistoryCollector, bool bValidateResultAgainstAvailabilities)
 {
 	using namespace UE::Chimera;
@@ -598,9 +597,9 @@ void UChimeraSubsystem::Query_AnyThread(const TArrayView<const FChimeraAvailabil
 	{
 		// if we find AnimInstance in an island, we perform ALL the Island motion matching searches.
 		// thread safety is ensured by the lock on Island->SearchResultsMutex in DoSearch_AnyThread
-		if (UChimeraIslandComponent* Island = FindIsland(AnimInstance))
+		if (FIsland* Island = FindIsland(AnimInstance))
 		{
-			Island->DoSearch_AnyThread(AnimInstance, Result);
+			Island->DoSearch_AnyThread(AnimInstance, ContinuingProperties, Result);
 
 			if (bValidateResultAgainstAvailabilities && Result.SelectedAnimation)
 			{
@@ -649,12 +648,11 @@ void UChimeraSubsystem::DebugDraw() const
 	static const int32 NumColors = sizeof(Colors) / sizeof(Colors[0]);
 	int32 CurrentColorIndex = 0;
 
-	for (const TObjectPtr<UChimeraIslandComponent>& IslandPtr : Islands)
+	for (const FIsland* Island : Islands)
 	{
-		UChimeraIslandComponent& Island = *IslandPtr.Get();
-		if (!Island.IsUninjected())
+		if (!Island->IsUninjected())
 		{
-			Island.DebugDraw(Colors[CurrentColorIndex]);
+			Island->DebugDraw(Colors[CurrentColorIndex]);
 			CurrentColorIndex = (CurrentColorIndex + 1) % NumColors;
 		}
 	}
@@ -672,14 +670,13 @@ void UChimeraSubsystem::DebugDraw() const
 			// looking for valid results for this AnimInstance
 			const UPoseSearchDatabase* ResultDatabase = nullptr;
 			UE::PoseSearch::FRole ResultRole;
-			for (const TObjectPtr<UChimeraIslandComponent>& IslandPtr : Islands)
+			for (const FIsland* Island : Islands)
 			{
-				UChimeraIslandComponent& Island = *IslandPtr.Get();
-				if (!Island.IsUninjected())
+				if (!Island->IsUninjected())
 				{
-					for (const FSearchResult& SearchResult : Island.GetSearchResults())
+					for (const FSearchResult& SearchResult : Island->GetSearchResults())
 					{
-						const FSearchContext& SearchContext = Island.GetSearchContexts()[SearchResult.SearchIndex];
+						const FSearchContext& SearchContext = Island->GetSearchContexts()[SearchResult.SearchIndex];
 						for (int32 AnimInstanceIndex = 0; AnimInstanceIndex < SearchContext.AnimInstances.Num(); ++AnimInstanceIndex)
 						{
 							if (SearchContext.AnimInstances[AnimInstanceIndex].Get() == AnimInstance)
diff --git a/Engine/Plugins/Experimental/Animation/Chimera/Source/Runtime/Public/Chimera/AnimNode_Chimera.h b/Engine/Plugins/Experimental/Animation/Chimera/Source/Runtime/Public/Chimera/AnimNode_Chimera.h
--- a/Engine/Plugins/Experimental/Animation/Chimera/Source/Runtime/Public/Chimera/AnimNode_Chimera.h
+++ b/Engine/Plugins/Experimental/Animation/Chimera/Source/Runtime/Public/Chimera/AnimNode_Chimera.h
@@ -25,7 +25,7 @@ enum class EChimeraEvaluationMode
 	// UntilContinuingPoseIsValid,
 };
 
-USTRUCT(BlueprintInternalUseOnly, Experimental)
+USTRUCT(BlueprintInternalUseOnly)
 struct CHIMERA_API FAnimNode_Chimera : public FAnimNode_BlendStack
 {
 	GENERATED_BODY();
@@ -65,7 +65,10 @@ protected:
 	float BlendLerp = 0.f;
 	float TranslationWarpLerp = 0.f;
 	float RotationWarpLerp = 0.f;
-	bool bWasInteracting = false;
+	// if a search is successful InteractingRolesNum > 0. 
+	// if InteractingRolesNum == 1 it means that the search is a regular single character motion matching search
+	// if InteractingRolesNum > 1 it means this node is interacting with other actors via the ChimeraSubsystem
+	int32 InteractingRolesNum = 0;
 
 	// Update Counter for detecting being relevant
 	FGraphTraversalCounter UpdateCounter;
diff --git a/Engine/Plugins/Experimental/Animation/Chimera/Source/Runtime/Public/Chimera/ChimeraAnimNodeLibrary.h b/Engine/Plugins/Experimental/Animation/Chimera/Source/Runtime/Public/Chimera/ChimeraAnimNodeLibrary.h
new file mode 100644
--- /dev/null
+++ b/Engine/Plugins/Experimental/Animation/Chimera/Source/Runtime/Public/Chimera/ChimeraAnimNodeLibrary.h
@@ -0,0 +1,39 @@
+// Copyright Epic Games, Inc. All Rights Reserved.
+
+#pragma once
+
+#include "Animation/AnimNodeReference.h"
+#include "Chimera/ChimeraLibrary.h"
+#include "Kismet/BlueprintFunctionLibrary.h"
+#include "ChimeraAnimNodeLibrary.generated.h"
+
+struct FAnimNode_Chimera;
+
+USTRUCT(BlueprintType)
+struct FChimeraAnimNodeReference : public FAnimNodeReference
+{
+	GENERATED_BODY()
+
+	typedef FAnimNode_Chimera FInternalNodeType;
+};
+
+UCLASS()
+class CHIMERA_API UChimeraAnimNodeLibrary : public UBlueprintFunctionLibrary
+{
+	GENERATED_BODY()
+
+public:
+	UFUNCTION(BlueprintCallable, Category = "Animation|Chimera", meta = (BlueprintThreadSafe, ExpandEnumAsExecs = "Result"))
+	static FChimeraAnimNodeReference ConvertToChimeraNode(const FAnimNodeReference& Node, EAnimNodeReferenceConversionResult& Result);
+
+	UFUNCTION(BlueprintPure, Category = "Animation|Chimera", meta = (BlueprintThreadSafe, DisplayName = "Convert to Chimera Node"))
+	static void ConvertToChimeraNodePure(const FAnimNodeReference& Node, FChimeraAnimNodeReference& ChimeraNode, bool& Result)
+	{
+		EAnimNodeReferenceConversionResult ConversionResult;
+		ChimeraNode = ConvertToChimeraNode(Node, ConversionResult);
+		Result = (ConversionResult == EAnimNodeReferenceConversionResult::Succeeded);
+	}
+
+	UFUNCTION(BlueprintCallable, Category = "Animation|Chimera", meta = (BlueprintThreadSafe))
+	static void SetAvailabilities(const FChimeraAnimNodeReference& ChimeraNode, const TArray<FChimeraAvailability>& Availabilities);
+};
\ No newline at end of file
diff --git a/Engine/Plugins/Experimental/Animation/Chimera/Source/Runtime/Public/Chimera/ChimeraAsset.h b/Engine/Plugins/Experimental/Animation/Chimera/Source/Runtime/Public/Chimera/ChimeraAsset.h
--- a/Engine/Plugins/Experimental/Animation/Chimera/Source/Runtime/Public/Chimera/ChimeraAsset.h
+++ b/Engine/Plugins/Experimental/Animation/Chimera/Source/Runtime/Public/Chimera/ChimeraAsset.h
@@ -7,7 +7,7 @@
 #include "PoseSearch/PoseSearchRole.h"
 #include "ChimeraAsset.generated.h"
 
-USTRUCT(Experimental)
+USTRUCT()
 struct CHIMERA_API FChimeraAssetItem
 {
 	GENERATED_BODY()
@@ -37,7 +37,7 @@ struct CHIMERA_API FChimeraAssetItem
 	FTransform Origin = FTransform::Identity;
 };
 
-UCLASS(Experimental, BlueprintType, Category = "Animation|Chimera")
+UCLASS(BlueprintType, Category = "Animation|Chimera")
 class CHIMERA_API UChimeraAsset : public UMultiAnimAsset
 {
 	GENERATED_BODY()
@@ -82,7 +82,7 @@ public:
 	FTransform GetDebugWarpOrigin(const UE::PoseSearch::FRole& Role, bool bComposeWithDebugWarpOffset) const;
 #endif // WITH_EDITOR
 
-	void CalculateWarpTransforms(float Time, const TArrayView<const FTransform> ActorRootBoneTransforms, TArrayView<FTransform> FullAlignedActorRootBoneTransforms, const UWorld* DebugDrawWorld = nullptr) const;
+	virtual void CalculateWarpTransforms(float Time, const TArrayView<const FTransform> ActorRootBoneTransforms, TArrayView<FTransform> FullAlignedActorRootBoneTransforms, const UWorld* DebugDrawWorld = nullptr) const override;
 
 	FQuat FindReferenceOrientation(const TArrayView<const FTransform> Transforms, const TArrayView<int32> SortedByWarpingWeightRotationItemIndex) const;
 	FVector FindReferencePosition(const TArrayView<const FTransform> Transforms, const TArrayView<float> NormalizedWarpingWeightTranslation) const;
diff --git a/Engine/Plugins/Experimental/Animation/Chimera/Source/Runtime/Public/Chimera/ChimeraIslandComponent.h b/Engine/Plugins/Experimental/Animation/Chimera/Source/Runtime/Public/Chimera/ChimeraIslandComponent.h
--- a/Engine/Plugins/Experimental/Animation/Chimera/Source/Runtime/Public/Chimera/ChimeraIslandComponent.h
+++ b/Engine/Plugins/Experimental/Animation/Chimera/Source/Runtime/Public/Chimera/ChimeraIslandComponent.h
@@ -2,9 +2,7 @@
 
 #pragma once
 
-#include "Components/ActorComponent.h"
 #include "PoseSearch/PoseSearchLibrary.h"
-#include "ChimeraIslandComponent.generated.h"
 
 struct FChimeraBlueprintResult;
 struct FPoseSearchContinuingProperties;
@@ -12,68 +10,81 @@ class UCharacterMovementComponent;
 
 namespace UE::Chimera
 {
-	struct FSearchContext
-	{
-		// @todo: since AnimInstances, HistoryCollectors, Roles, and BroadPhaseRadiuses have the same cardinality,
-		// add an array of structs instead of many arrays
-		TArray<TWeakObjectPtr<UAnimInstance>> AnimInstances;
-		// @todo: perhaps use the animnode of the pose history collector instead
-		TArray<const FAnimNode_PoseSearchHistoryCollector_Base*> HistoryCollectors;
-		TArray<UE::PoseSearch::FRole> Roles;
+struct FIsland;
+
+struct FSearchContext
+{
+	// @todo: since AnimInstances, HistoryCollectors, Roles, and BroadPhaseRadiuses have the same cardinality,
+	// add an array of structs instead of many arrays
+	TArray<TWeakObjectPtr<UAnimInstance>> AnimInstances;
+	// @todo: perhaps use the animnode of the pose history collector instead
+	TArray<const FAnimNode_PoseSearchHistoryCollector_Base*> HistoryCollectors;
+	TArray<UE::PoseSearch::FRole> Roles;
 
 #if ENABLE_DRAW_DEBUG
-		TArray<float> BroadPhaseRadiuses;
+	TArray<float> BroadPhaseRadiuses;
 #endif // ENABLE_DRAW_DEBUG
 
-		float MaxCost = MAX_flt;
+	float MaxCost = MAX_flt;
 
-		TWeakObjectPtr<const UPoseSearchDatabase> Database;
-		FPoseSearchContinuingProperties ContinuingProperties;
+	TWeakObjectPtr<const UPoseSearchDatabase> Database;
+	FPoseSearchContinuingProperties ContinuingProperties;
 
-		bool IsValid() const;
-		bool IsEquivalent(const FSearchContext& Other) const;
-	};
+	bool IsValid() const;
+	bool IsEquivalent(const FSearchContext& Other) const;
+};
 
-	struct FSearchResult : public UE::PoseSearch::FSearchResult
-	{
-		int32 SearchIndex = INDEX_NONE;
-		TArray<FTransform, TInlineAllocator<UE::PoseSearch::PreallocatedRolesNum>> FullAlignedActorRootBoneTransforms;
-	};
-} // namespace UE::Chimera
+struct FSearchResult : public UE::PoseSearch::FSearchResult
+{
+	int32 SearchIndex = INDEX_NONE;
+	TArray<FTransform, TInlineAllocator<UE::PoseSearch::PreallocatedRolesNum>> FullAlignedActorRootBoneTransforms;
 
-// UChimeraIslandComponent is a component injected between the interacting actors UCharacterMovementComponent and USkeletalMeshComponent
+	bool operator==(const FSearchResult& Other) const;
+};
+
+// FIsland contains ticks functions injected between the interacting actors UCharacterMovementComponent and USkeletalMeshComponent
 // to create a execution threading fence to be able to perform motion matching searches between the involved characters in a thread safe manner.
 // Look at UChimeraSubsystem "Execution model and threading details" for additional information
-UCLASS(Experimental, Category="Animation|Chimera")
-class UChimeraIslandComponent : public UActorComponent
+struct FIsland
 {
-public:
-	bool DoSearch_AnyThread(UObject* AnimInstance, FChimeraBlueprintResult& Result);
-
-private:
-	// to minimize the usage of lock we expose only thread safe APIs and allow private calls only via the power of friendship
-	friend class UChimeraSubsystem;
+	UE_NONCOPYABLE(FIsland);
+	
+	FIsland(ULevel* Level);
+	~FIsland();
 
-	GENERATED_BODY()
+	bool DoSearch_AnyThread(UObject* AnimInstance, const FPoseSearchContinuingProperties& ContinuingProperties, FChimeraBlueprintResult& Result);
 
-	virtual void TickComponent(float DeltaTime, enum ELevelTick TickType, FActorComponentTickFunction *ThisTickFunction) override;
+	const TArray<TWeakObjectPtr<UCharacterMovementComponent>>& GetCharacterMovementComponents() const { return CharacterMovementComponents; }
+	const TArray<TWeakObjectPtr<USkeletalMeshComponent>>& GetSkeletalMeshComponents() const { return SkeletalMeshComponents; }
+	const TArray<UE::Chimera::FSearchContext>& GetSearchContexts() const { return SearchContexts; }
+	const TArray<UE::Chimera::FSearchResult>& GetSearchResults() const { return SearchResults; }
 
+	bool IsUninjected() const;
 	void InjectToActor(AActor* Actor);
 	void Uninject();
-	bool IsUninjected();
 
+	const UE::Chimera::FSearchResult* FindSearchResult(const UE::Chimera::FSearchContext& SearchContext) const;
 	void AddSearchContext(const UE::Chimera::FSearchContext& SearchContext);
 
-	bool GetResult_AnyThread(UObject* AnimInstance, FChimeraBlueprintResult& Result);
-
 	void DebugDraw(const FColor& Color = FColor::Red) const;
 
-	const TArray<TWeakObjectPtr<UCharacterMovementComponent>>& GetCharacterMovementComponents() const { return CharacterMovementComponents; }
-	const TArray<TWeakObjectPtr<USkeletalMeshComponent>>& GetSkeletalMeshComponents() const { return SkeletalMeshComponents; }
-	const TArray<UE::Chimera::FSearchContext>& GetSearchContexts() const { return SearchContexts; }
-	const TArray<UE::Chimera::FSearchResult>& GetSearchResults() const { return SearchResults; }
+private:
+	bool GetResult_AnyThread(UObject* AnimInstance, FChimeraBlueprintResult& Result);
 
-	const UE::Chimera::FSearchResult* FindSearchResult(const UE::Chimera::FSearchContext& SearchContext) const;
+	struct FPreTickFunction : public FTickFunction
+	{
+		virtual void ExecuteTick(float DeltaTime, enum ELevelTick TickType, ENamedThreads::Type CurrentThread, const FGraphEventRef& MyCompletionGraphEvent) override;
+		virtual FString DiagnosticMessage() override { return TEXT("FPreTickFunction"); }
+		FIsland* Island = nullptr;
+	};
+	FPreTickFunction PreTickFunction;
+
+	struct FPostTickFunction : public FTickFunction
+	{
+		virtual void ExecuteTick(float DeltaTime, enum ELevelTick TickType, ENamedThreads::Type CurrentThread, const FGraphEventRef& MyCompletionGraphEvent) override;
+		virtual FString DiagnosticMessage() override { return TEXT("FPostTickFunction"); }
+	};
+	FPostTickFunction PostTickFunction;
 
 	TArray<TWeakObjectPtr<UCharacterMovementComponent>> CharacterMovementComponents;
 	TArray<TWeakObjectPtr<USkeletalMeshComponent>> SkeletalMeshComponents;
@@ -87,3 +98,5 @@ private:
 
 	FCriticalSection SearchResultsMutex;
 };
+
+} // namespace UE::Chimera
diff --git a/Engine/Plugins/Experimental/Animation/Chimera/Source/Runtime/Public/Chimera/ChimeraLibrary.h b/Engine/Plugins/Experimental/Animation/Chimera/Source/Runtime/Public/Chimera/ChimeraLibrary.h
--- a/Engine/Plugins/Experimental/Animation/Chimera/Source/Runtime/Public/Chimera/ChimeraLibrary.h
+++ b/Engine/Plugins/Experimental/Animation/Chimera/Source/Runtime/Public/Chimera/ChimeraLibrary.h
@@ -3,14 +3,13 @@
 #pragma once
 
 #include "Kismet/BlueprintFunctionLibrary.h"
+#include "PoseSearch/PoseSearchLibrary.h"
 #include "ChimeraLibrary.generated.h"
 
 // (CH)aracter (I)nteraction (M)atching (ERA) =):->
 // subsystem in charge of coordinating Characters availability for motion matched interactions,
 // scheduling motion matching searches and synchronizing animations playback
 
-struct FAnimNode_PoseSearchHistoryCollector_Base;
-class UPoseSearchDatabase;
 
 // input for ChimeraQuery_Pure: it declares that the associated character (AnimInstance) is willing to partecipate in an interction 
 // described by a UMultiAnimAsset (derived by UChimeraAsset) contained in the UPoseSearchDatabase Database
@@ -49,7 +48,7 @@ public:
 
 	// animation assigned to this character to partecipate in the interaction
 	UPROPERTY(Transient, VisibleAnywhere, BlueprintReadOnly, Category=State)
-	TObjectPtr<const UObject> SelectedAnimation = nullptr;
+	TObjectPtr<UObject> SelectedAnimation = nullptr;
 	
 	// SelectedAnimation associated time
 	UPROPERTY(Transient, VisibleAnywhere, BlueprintReadOnly, Category=State)
@@ -105,15 +104,20 @@ public:
 	// 
 	// if FChimeraBlueprintResult has a valid SelectedAnimation, this will be the animation assigned to this character to partecipate in this interaction.
 	// additional interaction properties, like assigned role, SelectedAnimation time, SearchCost, etc can be found within the result
+	// ContinuingProperties are used to figure out the continuing pose and bias it accordingly. ContinuingProperties can reference directly the UMultiAnimAsset
+	// or any of the roled UMultiAnimAsset::GetAnimationAsset, and the UChimeraSubsystem will figure out the related UMultiAnimAsset
 	// PoseHistoryName is the name of the pose history node used for the associated motion matching search
 	// if bValidateResultAgainstAvailabilities is true, the result will be invalidated if doesn't respect the new availabilities
 	UFUNCTION(BlueprintPure, Category = "Animation|Chimera", meta = (BlueprintThreadSafe))
-	static FChimeraBlueprintResult ChimeraQuery_Pure(TArray<FChimeraAvailability> Availabilities, UObject* AnimInstance, FName PoseHistoryName, bool bValidateResultAgainstAvailabilities);
+	static FChimeraBlueprintResult ChimeraQuery_Pure(TArray<FChimeraAvailability> Availabilities, UObject* AnimInstance, FPoseSearchContinuingProperties ContinuingProperties, FName PoseHistoryName, bool bValidateResultAgainstAvailabilities);
 
 	// BlueprintCallable version of ChimeraQuery_Pure
 	UFUNCTION(BlueprintCallable, Category = "Animation|Chimera", meta = (BlueprintThreadSafe))
-	static FChimeraBlueprintResult ChimeraQuery(TArray<FChimeraAvailability> Availabilities, UObject* AnimInstance, FName PoseHistoryName, bool bValidateResultAgainstAvailabilities);
+	static FChimeraBlueprintResult ChimeraQuery(TArray<FChimeraAvailability> Availabilities, UObject* AnimInstance, FPoseSearchContinuingProperties ContinuingProperties, FName PoseHistoryName, bool bValidateResultAgainstAvailabilities);
 
-	static FChimeraBlueprintResult ChimeraQuery(const TArrayView<const FChimeraAvailability> Availabilities, UObject* AnimInstance, const FAnimNode_PoseSearchHistoryCollector_Base* HistoryCollector, bool bValidateResultAgainstAvailabilities);
+	// version of ChimeraQuery_Pure referencing directly the HistoryCollector rather than looking for it by name
+	static FChimeraBlueprintResult ChimeraQuery(const TArrayView<const FChimeraAvailability> Availabilities, UObject* AnimInstance, const FPoseSearchContinuingProperties& ContinuingProperties, const FAnimNode_PoseSearchHistoryCollector_Base* HistoryCollector, bool bValidateResultAgainstAvailabilities);
 
+	UFUNCTION(BlueprintPure, Category = "Animation|Chimera", meta = (BlueprintThreadSafe))
+	static FPoseSearchContinuingProperties GetMontageContinuingProperties(UAnimInstance* AnimInstance);
 };
\ No newline at end of file
diff --git a/Engine/Plugins/Experimental/Animation/Chimera/Source/Runtime/Public/Chimera/ChimeraSubsystem.h b/Engine/Plugins/Experimental/Animation/Chimera/Source/Runtime/Public/Chimera/ChimeraSubsystem.h
--- a/Engine/Plugins/Experimental/Animation/Chimera/Source/Runtime/Public/Chimera/ChimeraSubsystem.h
+++ b/Engine/Plugins/Experimental/Animation/Chimera/Source/Runtime/Public/Chimera/ChimeraSubsystem.h
@@ -6,12 +6,13 @@
 #include "Chimera/ChimeraLibrary.h"
 #include "ChimeraSubsystem.generated.h"
 
+struct FPoseSearchContinuingProperties;
 class UAnimInstance;
-class UChimeraIslandComponent;
 class UPoseSearchDatabase;
 
 namespace UE::Chimera
 {
+	struct FIsland;
 	struct FSearchContext;
 
 	struct FChimeraAvailabilityEx : public FChimeraAvailability
@@ -31,24 +32,22 @@ namespace UE::Chimera
 		const FAnimNode_PoseSearchHistoryCollector_Base* HistoryCollector = nullptr;
 	};
 	typedef TMap<TWeakObjectPtr<UObject>, TArray<FChimeraAvailabilityEx>> FAvailabilityRequestsMap;
-	typedef TArray<TObjectPtr<UChimeraIslandComponent>> FIslands;
-
 } // namespace UE::Chimera
 
 // World subsystem accepting the publication of characters (via their AnimInstance(s)) FChimeraAvailability, representing the characters willingness to partecipate in an 
 // interaction with other characters from the next frame forward via Query_AnyThread method.
 //
-// The same method will return the FChimeraBlueprintResult from the PREVIOUS Tick processing (categorization of FChimeraAvailability(s) in multiple UChimeraIslandComponent(s)),
+// The same method will return the FChimeraBlueprintResult from the PREVIOUS Tick processing (categorization of FChimeraAvailability(s) in multiple FIsland(s)),
 // to the requesting character, containing the animation to play at what time, and the assigned role to partecipate in the selected interaction within the assigned 
-// UChimeraIslandComponent
+// FIsland
 
 // Execution model and threading details:
 /////////////////////////////////////////
 // 
 // - by calling UChimeraLibrary::ChimeraQuery_Pure(TArray<FChimeraAvailability> Availabilities, UObject* AnimInstance), characters publish their availabilities
 //   to partecipate in interactions to the UChimeraSubsystem
-// - UChimeraSubsystem::Tick processes those FChimeraAvailability(s) and creates/updates UChimeraIslandComponent. For each UChimeraIslandComponent it injects 
-//   a tick prerequisite via UChimeraIslandComponent::InjectToActor (that calls AddTickPrerequisiteComponent) to all the Actors in the same island.
+// - UChimeraSubsystem::Tick processes those FChimeraAvailability(s) and creates/updates UE::Chimera::FIsland. For each FIsland it injects 
+//   a tick prerequisite via FIsland::InjectToActor (that calls AddPrerequisite) to all the Actors in the same island.
 //   NoTe: the next frame the execution will be:
 //			for each island[k]
 //			{
@@ -57,21 +56,26 @@ namespace UE::Chimera
 //					Tick CharacterMovementComponent[k][i]
 //				}
 // 
-//				Tick ChimeraIslandComponent[k]
+//				Tick Island[k].PreTickFunction (that eventually generates the trajectories with all the updated CMCs)
 //
+//				Tick Actor[k][0].SkeletalMeshComponent (that performs the chimera queries for all the involved actors via DoSearch_AnyThread)
+// 
+//				Tick Island[k].PostTickFunction (currently just a threading fence for the execution of all the other SkeletalMeshComponent(s))
+// 
 //				for each Actor[k][i]
 //				{
-//					Tick SkeletalMeshComponent[k][i]
+//					if (i != 0)
+//						Tick SkeletalMeshComponent[k][i] (that DoSearch_AnyThread get the cached result calculated by Tick Actor[k][0].SkeletalMeshComponent)
 //				}
 //			}
 // - next frame UChimeraLibrary::ChimeraQuery_Pure(TArray<FChimeraAvailability> Availabilities, UObject* AnimInstance), with the context of all the published 
-//   availabilities and created islands, will find the associated UChimeraIslandComponent to the AnimInstance and call UChimeraIslandComponent::DoSearch_AnyThread 
+//   availabilities and created islands, will find the associated FIsland to the AnimInstance and call FIsland::DoSearch_AnyThread 
 //   (via UChimeraSubsystem::Query_AnyThread) that will perform ALL (YES, ALL, so the bigger the island the slower the execution) the motion matching searches
-//   for all the possible Actors / databases / Roles permutations, and populate UChimeraIslandComponent::SearchResults with ALL the results for the island.
+//   for all the possible Actors / databases / Roles permutations, and populate FIsland::SearchResults with ALL the results for the island.
 //   Ultimately the ChimeraQuery_Pure will return the SearchResults associated to the requesting AnimInstance with information about what animation to play
 //   at what time with wich Role.
 
-UCLASS(Experimental, Category="Animation|Chimera")
+UCLASS(Category="Animation|Chimera")
 class CHIMERA_API UChimeraSubsystem : public UTickableWorldSubsystem
 {
 	GENERATED_BODY()
@@ -81,7 +85,7 @@ public:
 
 	virtual void Deinitialize() override;
 
-	// it processes FChimeraAvailability(s) and creates/updates UChimeraIslandComponent
+	// it processes FChimeraAvailability(s) and creates/updates FIsland
 	virtual void Tick(float DeltaSeconds) override;
 	
 	virtual TStatId GetStatId() const override;
@@ -89,11 +93,12 @@ public:
 	// publishing FChimeraAvailability(s) for the requesting character (AnimInstance) and getting the FChimeraBlueprintResult from the PREVIOUS Tick update
 	// containing the animation to play at what time, and the assigned role to partecipate in the selected interaction
 	// Either a PoseHistoryName or a HistoryCollector are required to perform the associated motion matching searches
-	void Query_AnyThread(const TArrayView<const FChimeraAvailability> Availabilities, UObject* AnimInstance, FChimeraBlueprintResult& Result,
+	void Query_AnyThread(const TArrayView<const FChimeraAvailability> Availabilities, UObject* AnimInstance, 
+		const FPoseSearchContinuingProperties& ContinuingProperties, FChimeraBlueprintResult& Result,
 		FName PoseHistoryName, const FAnimNode_PoseSearchHistoryCollector_Base* HistoryCollector, bool bValidateResultAgainstAvailabilities);
 private:
-	UChimeraIslandComponent& CreateIsland();
-	UChimeraIslandComponent& GetAvailableIsland();
+	UE::Chimera::FIsland& CreateIsland();
+	UE::Chimera::FIsland& GetAvailableIsland();
 
 	void DestroyIsland(int32 Index);
 	void DestroyAllIslands();
@@ -101,7 +106,7 @@ private:
 	bool ValidateAllIslands() const;
 
 	void PopulateContinuingProperties(UE::Chimera::FSearchContext& SearchContext, float DeltaSeconds) const;
-	UChimeraIslandComponent* FindIsland(UObject* InAnimInstance);
+	UE::Chimera::FIsland* FindIsland(UObject* InAnimInstance);
 	
 	void DebugDraw() const;
 
@@ -109,7 +114,7 @@ private:
 	FCriticalSection AvailabilityRequestsMapMutex;
 
 	// array of groups of characters that needs to be anaylzed together for possible interactions
-	UE::Chimera::FIslands Islands;
+	TArray<UE::Chimera::FIsland*> Islands;
 
 	// critical section to retrieve the subsystem in a thread safe manner
 	static FCriticalSection RetrieveSubsystemMutex;
