chimera - UMultiAnimAsset::CalculateWarpTransforms clenaup

diff --git a/Engine/Plugins/Animation/PoseSearch/Source/Runtime/Public/PoseSearch/MultiAnimAsset.h b/Engine/Plugins/Animation/PoseSearch/Source/Runtime/Public/PoseSearch/MultiAnimAsset.h
--- a/Engine/Plugins/Animation/PoseSearch/Source/Runtime/Public/PoseSearch/MultiAnimAsset.h
+++ b/Engine/Plugins/Animation/PoseSearch/Source/Runtime/Public/PoseSearch/MultiAnimAsset.h
@@ -28,7 +28,7 @@ public:
 	[[nodiscard]] virtual UAnimationAsset* GetAnimationAsset(const FName& Role) const PURE_VIRTUAL(UMultiAnimAsset::GetAnimationAsset, return nullptr;);
 	[[nodiscard]] virtual FTransform GetOrigin(const FName& Role) const PURE_VIRTUAL(UMultiAnimAsset::GetOrigin, return FTransform::Identity;);
 
-	virtual void CalculateWarpTransforms(float Time, const TArrayView<const FTransform> ActorRootBoneTransforms, TArrayView<FTransform> FullAlignedActorRootBoneTransforms, const UWorld* DebugDrawWorld = nullptr) const PURE_VIRTUAL(UMultiAnimAsset::CalculateWarpTransforms, );
+	virtual void CalculateWarpTransforms(float Time, const TArrayView<const FTransform> ActorRootBoneTransforms, TArrayView<FTransform> FullAlignedActorRootBoneTransforms) const PURE_VIRTUAL(UMultiAnimAsset::CalculateWarpTransforms, );
 
 	UFUNCTION(BlueprintPure, Category = "Animation", meta=(BlueprintThreadSafe, DisplayName = "Get Animation Asset"))
 	UAnimationAsset* BP_GetAnimationAsset(const FName& Role) const { return GetAnimationAsset(Role); }
diff --git a/Engine/Plugins/Experimental/Animation/Chimera/Source/Editor/Private/ChimeraAssetEditor.cpp b/Engine/Plugins/Experimental/Animation/Chimera/Source/Editor/Private/ChimeraAssetEditor.cpp
--- a/Engine/Plugins/Experimental/Animation/Chimera/Source/Editor/Private/ChimeraAssetEditor.cpp
+++ b/Engine/Plugins/Experimental/Animation/Chimera/Source/Editor/Private/ChimeraAssetEditor.cpp
@@ -319,7 +319,6 @@ void FChimeraAssetViewModel::Tick(float DeltaSeconds)
 #if WITH_EDITORONLY_DATA
 	if (ChimeraAsset->bEnableDebugWarp && PreviewActors.Num() == ChimeraAsset->GetNumRoles())
 	{
-		const UWorld* DebugDrawWorld = nullptr;
 		TArray<FTransform, TInlineAllocator<UE::PoseSearch::PreallocatedRolesNum>> ActorTransforms;
 		ActorTransforms.SetNum(ChimeraAsset->GetNumRoles());
 		for (int32 ChimeraAssetRoleIndex = 0; ChimeraAssetRoleIndex < ChimeraAsset->GetNumRoles(); ++ChimeraAssetRoleIndex)
@@ -331,18 +330,11 @@ void FChimeraAssetViewModel::Tick(float DeltaSeconds)
 			check(PreviewActor.GetRole() == ChimeraAssetRole);
 
 			ActorTransforms[ChimeraAssetRoleIndex] = PreviewActor.GetDebugActorTransformFromSampler();
-			if (!DebugDrawWorld)
-			{
-				if (UAnimPreviewInstance* AnimPreviewInstance = PreviewActor.GetAnimPreviewInstance())
-				{
-					DebugDrawWorld = AnimPreviewInstance->GetWorld();
-				}
-			}
 		}
 
 		TArray<FTransform, TInlineAllocator<UE::PoseSearch::PreallocatedRolesNum>> FullAlignedActorTransforms;
 		FullAlignedActorTransforms.SetNum(ChimeraAsset->GetNumRoles());
-		ChimeraAsset->CalculateWarpTransforms(PlayTime, ActorTransforms, FullAlignedActorTransforms, DebugDrawWorld);
+		ChimeraAsset->CalculateWarpTransforms(PlayTime, ActorTransforms, FullAlignedActorTransforms);
 
 		for (int32 ChimeraAssetRoleIndex = 0; ChimeraAssetRoleIndex < ChimeraAsset->GetNumRoles(); ++ChimeraAssetRoleIndex)
 		{
diff --git a/Engine/Plugins/Experimental/Animation/Chimera/Source/Runtime/Private/ChimeraAsset.cpp b/Engine/Plugins/Experimental/Animation/Chimera/Source/Runtime/Private/ChimeraAsset.cpp
--- a/Engine/Plugins/Experimental/Animation/Chimera/Source/Runtime/Private/ChimeraAsset.cpp
+++ b/Engine/Plugins/Experimental/Animation/Chimera/Source/Runtime/Private/ChimeraAsset.cpp
@@ -3,16 +3,8 @@
 #include "Chimera/ChimeraAsset.h"
 #include "Animation/AnimSequence.h"
 #include "Animation/BlendSpace.h"
-#include "DrawDebugHelpers.h"
 #include "PoseSearch/PoseSearchAssetSampler.h"
 
-namespace UE::Chimera
-{
-#if ENABLE_DRAW_DEBUG && ENABLE_ANIM_DEBUG
-TAutoConsoleVariable<bool> CVarAnimChimeraDebugDrawWarpTransforms(TEXT("a.Chimera.DebugDrawWarpTransforms"), false, TEXT("Enable / Disable Warp Transforms DebugDraw"));
-#endif
-}
-
 bool UChimeraAsset::IsLooping() const
 {
 	float CommonPlayLength = -1.f;
@@ -229,7 +221,7 @@ FTransform UChimeraAsset::GetDebugWarpOrigin(const UE::PoseSearch::FRole& Role,
 }
 #endif // WITH_EDITOR
 
-void UChimeraAsset::CalculateWarpTransforms(float Time, const TArrayView<const FTransform> ActorRootBoneTransforms, TArrayView<FTransform> FullAlignedActorRootBoneTransforms, const UWorld* DebugDrawWorld) const
+void UChimeraAsset::CalculateWarpTransforms(float Time, const TArrayView<const FTransform> ActorRootBoneTransforms, TArrayView<FTransform> FullAlignedActorRootBoneTransforms) const
 {
 	check(ActorRootBoneTransforms.Num() == GetNumRoles());
 	check(FullAlignedActorRootBoneTransforms.Num() == GetNumRoles());
@@ -341,16 +333,6 @@ void UChimeraAsset::CalculateWarpTransforms(float Time, const TArrayView<const F
 	const FVector AssetReferencePosition = FindReferencePosition(AssetRootBoneTransforms, NormalizedWarpingWeightTranslation);
 	const FVector ActorsReferencePosition = FindReferencePosition(ActorRootBoneTransforms, NormalizedWarpingWeightTranslation);
 
-#if ENABLE_DRAW_DEBUG && ENABLE_ANIM_DEBUG
-	const bool bDrawWarpTransforms = UE::Chimera::CVarAnimChimeraDebugDrawWarpTransforms.GetValueOnAnyThread();
-	if (bDrawWarpTransforms)
-	{
-		DrawDebugCoordinateSystem(DebugDrawWorld, AssetReferencePosition, AssetReferenceOrientation.Rotator(), 10.f, false, 0.f, SDPG_Foreground);
-		DrawDebugCoordinateSystem(DebugDrawWorld, ActorsReferencePosition, ActorsReferenceOrientation.Rotator(), 10.f, false, 0.f, SDPG_Foreground);
-		DrawDebugCoordinateSystem(DebugDrawWorld, ActorsReferencePosition, WeightedActorsReferenceOrientation.Rotator(), 20.f, false, 0.f, SDPG_Foreground);
-	}
-#endif // ENABLE_DRAW_DEBUG && ENABLE_ANIM_DEBUG
-
 	// aligning all the actors to ActorsReferencePosition, WeightedActorsReferenceOrientation
 	const FTransform AssetReferenceTransform(AssetReferenceOrientation, AssetReferencePosition);
 	const FTransform ActorsReferenceTransform(WeightedActorsReferenceOrientation, ActorsReferencePosition);
@@ -359,12 +341,5 @@ void UChimeraAsset::CalculateWarpTransforms(float Time, const TArrayView<const F
 	for (int32 ItemIndex = 0; ItemIndex < ItemsNum; ++ItemIndex)
 	{
 		FullAlignedActorRootBoneTransforms[ItemIndex] = (AssetRootBoneTransforms[ItemIndex] * AssetReferenceInverseTransform) * ActorsReferenceTransform;
-
-#if ENABLE_DRAW_DEBUG && ENABLE_ANIM_DEBUG
-		if (bDrawWarpTransforms)
-		{
-			DrawDebugCoordinateSystem(DebugDrawWorld, FullAlignedActorRootBoneTransforms[ItemIndex].GetTranslation(), FullAlignedActorRootBoneTransforms[ItemIndex].Rotator(), 10.f, false, 0.f, SDPG_Foreground);
-		}
-#endif // ENABLE_DRAW_DEBUG && ENABLE_ANIM_DEBUG
 	}
 }
\ No newline at end of file
diff --git a/Engine/Plugins/Experimental/Animation/Chimera/Source/Runtime/Private/ChimeraIslandComponent.cpp b/Engine/Plugins/Experimental/Animation/Chimera/Source/Runtime/Private/ChimeraIslandComponent.cpp
--- a/Engine/Plugins/Experimental/Animation/Chimera/Source/Runtime/Private/ChimeraIslandComponent.cpp
+++ b/Engine/Plugins/Experimental/Animation/Chimera/Source/Runtime/Private/ChimeraIslandComponent.cpp
@@ -48,7 +48,6 @@ static FSearchResult InitSearchResult(const UE::PoseSearch::FSearchResult& Searc
 
 				// mapping MultiAnimAsset Roles to SearchContext.* indexes
 				TArray<FTransform, TInlineAllocator<UE::PoseSearch::PreallocatedRolesNum>> ActorRootBoneTransforms;
-				const UWorld* DebugDrawWorld = nullptr;
 				for (int32 MultiAnimAssetRoleIndex = 0; MultiAnimAssetRoleIndex < MultiAnimAsset->GetNumRoles(); ++MultiAnimAssetRoleIndex)
 				{
 					const UE::PoseSearch::FRole& ChimeraAssetRole = MultiAnimAsset->GetRole(MultiAnimAssetRoleIndex);
@@ -58,11 +57,6 @@ static FSearchResult InitSearchResult(const UE::PoseSearch::FSearchResult& Searc
 					{
 						const FTransform RootBoneTransform = AnimInstance->GetSkelMeshComponent()->GetBoneTransform(0);
 						ActorRootBoneTransforms.Add(RootBoneTransform);
-					
-						if (!DebugDrawWorld)
-						{
-							DebugDrawWorld = AnimInstance->GetWorld();
-						}
 					}
 				}
 
@@ -74,7 +68,7 @@ static FSearchResult InitSearchResult(const UE::PoseSearch::FSearchResult& Searc
 					FullAlignedActorRootBoneTransforms.SetNum(MultiAnimAsset->GetNumRoles());
 
 					// @todo: should it be SearchResult.AssetTime + DeltaTime?
-					MultiAnimAsset->CalculateWarpTransforms(SearchResult.AssetTime, ActorRootBoneTransforms, FullAlignedActorRootBoneTransforms, DebugDrawWorld);
+					MultiAnimAsset->CalculateWarpTransforms(SearchResult.AssetTime, ActorRootBoneTransforms, FullAlignedActorRootBoneTransforms);
 					ChimeraSearchResult.FullAlignedActorRootBoneTransforms.SetNum(MultiAnimAsset->GetNumRoles());
 
 					for (int32 ChimeraAssetRoleIndex = 0; ChimeraAssetRoleIndex < MultiAnimAsset->GetNumRoles(); ++ChimeraAssetRoleIndex)
diff --git a/Engine/Plugins/Experimental/Animation/Chimera/Source/Runtime/Public/Chimera/ChimeraAsset.h b/Engine/Plugins/Experimental/Animation/Chimera/Source/Runtime/Public/Chimera/ChimeraAsset.h
--- a/Engine/Plugins/Experimental/Animation/Chimera/Source/Runtime/Public/Chimera/ChimeraAsset.h
+++ b/Engine/Plugins/Experimental/Animation/Chimera/Source/Runtime/Public/Chimera/ChimeraAsset.h
@@ -82,7 +82,7 @@ public:
 	FTransform GetDebugWarpOrigin(const UE::PoseSearch::FRole& Role, bool bComposeWithDebugWarpOffset) const;
 #endif // WITH_EDITOR
 
-	virtual void CalculateWarpTransforms(float Time, const TArrayView<const FTransform> ActorRootBoneTransforms, TArrayView<FTransform> FullAlignedActorRootBoneTransforms, const UWorld* DebugDrawWorld = nullptr) const override;
+	virtual void CalculateWarpTransforms(float Time, const TArrayView<const FTransform> ActorRootBoneTransforms, TArrayView<FTransform> FullAlignedActorRootBoneTransforms) const override;
 
 	FQuat FindReferenceOrientation(const TArrayView<const FTransform> Transforms, const TArrayView<int32> SortedByWarpingWeightRotationItemIndex) const;
 	FVector FindReferencePosition(const TArrayView<const FTransform> Transforms, const TArrayView<float> NormalizedWarpingWeightTranslation) const;
