chimera - renaming ChimeraIslandComponent to ChimeraIsland

diff --git a/Engine/Plugins/Experimental/Animation/Chimera/Source/Runtime/Private/ChimeraIsland.cpp b/Engine/Plugins/Experimental/Animation/Chimera/Source/Runtime/Private/ChimeraIsland.cpp
new file mode 100644
--- /dev/null
+++ b/Engine/Plugins/Experimental/Animation/Chimera/Source/Runtime/Private/ChimeraIsland.cpp
@@ -0,0 +1,668 @@
+// Copyright Epic Games, Inc. All Rights Reserved.
+
+#include "Chimera/ChimeraIsland.h"
+#include "Animation/AnimClassInterface.h"
+#include "Animation/AnimInstance.h"
+#include "Components/SkeletalMeshComponent.h"
+#include "Chimera/ChimeraDefines.h"
+#include "Chimera/ChimeraLibrary.h"
+#include "DrawDebugHelpers.h"
+#include "GameFramework/CharacterMovementComponent.h"
+#include "PoseSearch/AnimNode_PoseSearchHistoryCollector.h"
+#include "PoseSearch/MultiAnimAsset.h"
+#include "PoseSearch/PoseSearchLibrary.h"
+#include "PoseSearch/PoseSearchDatabase.h"
+#include "PoseSearch/PoseSearchSchema.h"
+
+namespace UE::Chimera
+{
+
+#if ENABLE_ANIM_DEBUG
+static TAutoConsoleVariable<bool> CVarChimeraShowIslands(TEXT("a.Chimera.ShowIslands"), false, TEXT("Show Chimera Islands"));
+#endif
+
+// Utils functions
+///////////////////////////////////////////////////////////
+static FSearchResult InitSearchResult(const UE::PoseSearch::FSearchResult& SearchResult, int32 SearchIndex, const FSearchContext& SearchContext)
+{
+	FSearchResult ChimeraSearchResult;
+	static_cast<UE::PoseSearch::FSearchResult&>(ChimeraSearchResult) = SearchResult;
+	ChimeraSearchResult.SearchIndex = SearchIndex;
+
+	if (const UE::PoseSearch::FSearchIndexAsset* SearchIndexAsset = SearchResult.GetSearchIndexAsset())
+	{
+		if (const FPoseSearchDatabaseAnimationAssetBase* DatabaseAnimationAssetBase = SearchResult.Database->GetDatabaseAnimationAsset<FPoseSearchDatabaseAnimationAssetBase>(*SearchIndexAsset))
+		{
+			if (const UMultiAnimAsset* MultiAnimAsset = Cast<UMultiAnimAsset>(DatabaseAnimationAssetBase->GetAnimationAsset()))
+			{
+				check(MultiAnimAsset->GetNumRoles() == SearchContext.AnimInstances.Num());
+				check(MultiAnimAsset->GetNumRoles() == SearchContext.HistoryCollectors.Num());
+				check(MultiAnimAsset->GetNumRoles() == SearchContext.Roles.Num());
+
+				UE::PoseSearch::FRoleToIndex SearchContextRoleToIndex;
+				SearchContextRoleToIndex.Reserve(SearchContext.Roles.Num());
+				for (int32 SearchContextRoleIndex = 0; SearchContextRoleIndex < SearchContext.Roles.Num(); ++SearchContextRoleIndex)
+				{
+					SearchContextRoleToIndex.Add(SearchContext.Roles[SearchContextRoleIndex]) = SearchContextRoleIndex;
+				}
+
+				// mapping MultiAnimAsset Roles to SearchContext.* indexes
+				TArray<FTransform, TInlineAllocator<UE::PoseSearch::PreallocatedRolesNum>> ActorRootBoneTransforms;
+				for (int32 MultiAnimAssetRoleIndex = 0; MultiAnimAssetRoleIndex < MultiAnimAsset->GetNumRoles(); ++MultiAnimAssetRoleIndex)
+				{
+					const UE::PoseSearch::FRole& ChimeraAssetRole = MultiAnimAsset->GetRole(MultiAnimAssetRoleIndex);
+					const int32 SearchContextIndex = SearchContextRoleToIndex[ChimeraAssetRole];
+
+					if (UAnimInstance* AnimInstance = SearchContext.AnimInstances[SearchContextIndex].Get())
+					{
+						const FTransform RootBoneTransform = AnimInstance->GetSkelMeshComponent()->GetBoneTransform(0);
+						ActorRootBoneTransforms.Add(RootBoneTransform);
+					}
+				}
+
+				if (ActorRootBoneTransforms.Num() == MultiAnimAsset->GetNumRoles())
+				{
+					// FullAlignedActorRootBoneTransforms is mapped to the MultiAnimAsset roles:
+					// FullAlignedActorRootBoneTransforms[0] is for MultiAnimAsset->GetRole(0)
+					TArray<FTransform, TInlineAllocator<UE::PoseSearch::PreallocatedRolesNum>> FullAlignedActorRootBoneTransforms;
+					FullAlignedActorRootBoneTransforms.SetNum(MultiAnimAsset->GetNumRoles());
+
+					// @todo: should it be SearchResult.AssetTime + DeltaTime?
+					MultiAnimAsset->CalculateWarpTransforms(SearchResult.AssetTime, ActorRootBoneTransforms, FullAlignedActorRootBoneTransforms);
+					ChimeraSearchResult.FullAlignedActorRootBoneTransforms.SetNum(MultiAnimAsset->GetNumRoles());
+
+					for (int32 ChimeraAssetRoleIndex = 0; ChimeraAssetRoleIndex < MultiAnimAsset->GetNumRoles(); ++ChimeraAssetRoleIndex)
+					{
+						const UE::PoseSearch::FRole& ChimeraAssetRole = MultiAnimAsset->GetRole(ChimeraAssetRoleIndex);
+						const int32 SearchContextIndex = SearchContextRoleToIndex[ChimeraAssetRole];
+
+						ChimeraSearchResult.FullAlignedActorRootBoneTransforms[SearchContextIndex] = FullAlignedActorRootBoneTransforms[ChimeraAssetRoleIndex];
+					}
+				}
+			}
+			else
+			{
+				// support for non UMultiAnimAsset to be backward compatible with regular motion matching searches
+				check(SearchContext.AnimInstances.Num() == 1);
+
+				if (UAnimInstance* AnimInstance = SearchContext.AnimInstances[0].Get())
+				{
+					const FTransform RootBoneTransform = AnimInstance->GetSkelMeshComponent()->GetBoneTransform(0);
+					ChimeraSearchResult.FullAlignedActorRootBoneTransforms.SetNum(1);
+					// @todo: should it be RootBoneTransform + root motion transform?
+					ChimeraSearchResult.FullAlignedActorRootBoneTransforms[0] = RootBoneTransform;
+				}
+			}
+		}
+	}
+
+	return ChimeraSearchResult;
+}
+
+static void InitSearchResults(TArray<FSearchResult>& SearchResults, TConstArrayView<UE::PoseSearch::FSearchResult> PoseSearchResults, TConstArrayView<FSearchContext> SearchContexts)
+{
+	// WIP!
+	// @todo: figure out multiple policies to use the most characters? right now only the best search is "valid" with the most characters
+	int32 BestSearchIndex = INDEX_NONE;
+	for (int32 SearchIndex = 0; SearchIndex < PoseSearchResults.Num(); ++SearchIndex)
+	{
+		if (PoseSearchResults[SearchIndex].IsValid())
+		{
+			if (BestSearchIndex == INDEX_NONE)
+			{
+				BestSearchIndex = SearchIndex;
+			}
+			else if (SearchContexts[SearchIndex].Roles.Num() > SearchContexts[BestSearchIndex].Roles.Num())
+			{
+				BestSearchIndex = SearchIndex;
+			}
+			else if (SearchContexts[SearchIndex].Roles.Num() == SearchContexts[BestSearchIndex].Roles.Num() &&
+				PoseSearchResults[SearchIndex].PoseCost < PoseSearchResults[BestSearchIndex].PoseCost)
+			{
+				BestSearchIndex = SearchIndex;
+			}
+		}
+	}
+
+	if (BestSearchIndex != INDEX_NONE)
+	{
+		SearchResults.SetNum(1);
+		SearchResults[0] = InitSearchResult(PoseSearchResults[BestSearchIndex], BestSearchIndex, SearchContexts[BestSearchIndex]);
+	}
+}
+
+// FSearchResult
+///////////////////////////////////////////////////////////
+bool FSearchContext::IsValid() const
+{
+	if (Database == nullptr)
+	{
+		return false;
+	}
+
+	const int32 Num = AnimInstances.Num();
+	if (Num < 1)
+	{
+		return false;
+	}
+
+	if (Num != HistoryCollectors.Num())
+	{
+		return false;
+	}
+
+	if (Num != Roles.Num())
+	{
+		return false;
+	}
+
+	for (int32 IndexA = 0; IndexA < Num; ++IndexA)
+	{
+		if (AnimInstances[IndexA] == nullptr)
+		{
+			return false;
+		}
+
+		for (int32 IndexB = IndexA + 1; IndexB < Num; ++IndexB)
+		{
+			if (AnimInstances[IndexA] == AnimInstances[IndexB])
+			{
+				return false;
+			}
+		}
+	}
+
+	for (int32 IndexA = 0; IndexA < Num; ++IndexA)
+	{
+		for (int32 IndexB = IndexA + 1; IndexB < Num; ++IndexB)
+		{
+			if (Roles[IndexA] == Roles[IndexB])
+			{
+				return false;
+			}
+		}
+	}
+
+	for (int32 IndexA = 0; IndexA < Num; ++IndexA)
+	{
+		if (HistoryCollectors[IndexA] == nullptr)
+		{
+			return false;
+		}
+
+		for (int32 IndexB = IndexA + 1; IndexB < Num; ++IndexB)
+		{
+			if (HistoryCollectors[IndexA] == HistoryCollectors[IndexB])
+			{
+				return false;
+			}
+		}
+	}
+
+	return true;
+}
+
+bool FSearchContext::IsEquivalent(const FSearchContext& Other) const
+{
+	if (Database != Other.Database)
+	{
+		return false;
+	}
+
+	const int32 Num = AnimInstances.Num();
+	if (Num != Other.AnimInstances.Num())
+	{
+		return false;
+	}
+
+	int32 CommonRoledAnimInstances = 0;
+	for (int32 IndexThis = 0; IndexThis < Num; ++IndexThis)
+	{
+		for (int32 IndexOther = 0; IndexOther < Num; ++IndexOther)
+		{
+			if (AnimInstances[IndexThis] == Other.AnimInstances[IndexOther] &&
+				HistoryCollectors[IndexThis] == Other.HistoryCollectors[IndexOther] &&
+				Roles[IndexThis] == Other.Roles[IndexOther])
+			{
+				++CommonRoledAnimInstances;
+			}
+		}
+	}
+
+	// using >= in case there are duplicated animinstances in this or Other (this->IsValid() should be false!)
+	if (CommonRoledAnimInstances >= Num)
+	{
+		return true;
+	}
+
+	return false;
+}
+
+// FSearchResult
+///////////////////////////////////////////////////////////
+bool FSearchResult::operator==(const FSearchResult& Other) const
+{
+	return static_cast<const UE::PoseSearch::FSearchResult&>(*this) == static_cast<const UE::PoseSearch::FSearchResult&>(Other) &&
+				FullAlignedActorRootBoneTransforms.Num() == Other.FullAlignedActorRootBoneTransforms.Num() &&
+				0 == FMemory::Memcmp(FullAlignedActorRootBoneTransforms.GetData(), Other.FullAlignedActorRootBoneTransforms.GetData(), FullAlignedActorRootBoneTransforms.Num() * sizeof(FTransform));
+}
+
+// FIslandPreTickFunction
+///////////////////////////////////////////////////////////
+void FIsland::FPreTickFunction::ExecuteTick(float DeltaTime, enum ELevelTick TickType, ENamedThreads::Type CurrentThread, const FGraphEventRef& MyCompletionGraphEvent)
+{
+	// Called before any skeletal mesh component tick, when there aren't animation jobs flying. No need to FScopeLock Lock(&Mutex);
+	// generating trajectories before running any of the skeletal mesh component ticks
+	for (const FSearchContext& SearchContext : Island->SearchContexts)
+	{
+		for (int32 Index = 0; Index < SearchContext.AnimInstances.Num(); ++Index)
+		{
+			if (const UAnimInstance* AnimInstance = SearchContext.AnimInstances[Index].Get())
+			{
+				// since FIsland has a tick dependency with the USkeletalMeshComponent it's safe modify the FAnimNode_PoseSearchHistoryCollector_Base
+				FAnimNode_PoseSearchHistoryCollector_Base* HistoryCollector = const_cast<FAnimNode_PoseSearchHistoryCollector_Base*>(SearchContext.HistoryCollectors[Index]);
+				check(HistoryCollector);
+				HistoryCollector->GenerateTrajectory(AnimInstance);
+			}
+		}
+	}
+}
+
+// FPostTickFunction
+///////////////////////////////////////////////////////////
+void FIsland::FPostTickFunction::ExecuteTick(float DeltaTime, enum ELevelTick TickType, ENamedThreads::Type CurrentThread, const FGraphEventRef& MyCompletionGraphEvent)
+{
+	// do nothing
+}
+
+// FIsland
+///////////////////////////////////////////////////////////
+FIsland::FIsland(ULevel* Level)
+{
+	PreTickFunction.bCanEverTick = true;
+	PreTickFunction.bStartWithTickEnabled = true;
+	PreTickFunction.TickGroup = ETickingGroup::TG_PrePhysics;
+	PreTickFunction.Island = this;
+	PreTickFunction.SetTickFunctionEnable(true);
+	PreTickFunction.RegisterTickFunction(Level);
+
+	PostTickFunction.bCanEverTick = true;
+	PostTickFunction.bStartWithTickEnabled = true;
+	PostTickFunction.TickGroup = ETickingGroup::TG_PrePhysics;
+	PostTickFunction.SetTickFunctionEnable(true);
+	PostTickFunction.RegisterTickFunction(Level);
+}
+
+FIsland::~FIsland()
+{
+	PreTickFunction.UnRegisterTickFunction();
+	PostTickFunction.UnRegisterTickFunction();
+
+	Uninject();
+}
+
+void FIsland::DebugDraw(const FColor& Color) const
+{
+	// called only by UChimeraSubsystem::Tick so no need to lock SearchResultsMutex to protect the read of SearchContexts
+#if ENABLE_DRAW_DEBUG
+	check(IsInGameThread());
+
+	if (CVarChimeraShowIslands.GetValueOnAnyThread())
+	{
+		for (const FSearchContext& SearchContext : SearchContexts)
+		{
+			for (int32 Index = 0; Index < SearchContext.AnimInstances.Num(); ++Index)
+			{
+				if (const UAnimInstance* AnimInstance = SearchContext.AnimInstances[Index].Get())
+				{
+					const FVector Position = AnimInstance->GetSkelMeshComponent()->GetComponentLocation();
+					const float BroadPhaseRadius = SearchContext.BroadPhaseRadiuses[Index];
+					DrawDebugCircle(AnimInstance->GetWorld(), Position, BroadPhaseRadius, 40, Color, false, 0.f, SDPG_Foreground, 0.f, FVector::XAxisVector, FVector::YAxisVector, false);
+				}
+			}
+		}
+	}
+#endif // ENABLE_DRAW_DEBUG
+}
+
+void FIsland::InjectToActor(AActor* Actor)
+{
+	check(IsInGameThread());
+
+	// Called by UChimeraSubsystem::Tick when there aren't animation jobs flying. No need to FScopeLock Lock(&Mutex);
+	if (Actor)
+	{
+		if (UCharacterMovementComponent* CharacterMovementComponent = Actor->GetComponentByClass<UCharacterMovementComponent>())
+		{
+			if (USkeletalMeshComponent* SkeletalMeshComponent = Actor->GetComponentByClass<USkeletalMeshComponent>())
+			{
+				const bool bIsFirstInjectedActor = IsUninjected();
+
+				//	tick order: 
+				//		CharacterMovementComponent(s) ->
+				//			Island.PreTickFunction ->
+				//				first injected actor SkeletalMeshComponent ->
+				//					Island.PostTickFunction ->
+				//						other SkeletalMeshComponent(s)
+				CharacterMovementComponents.AddUnique(CharacterMovementComponent);
+				SkeletalMeshComponents.AddUnique(SkeletalMeshComponent);
+	
+				// making sure that if we add a unique CharacterMovementComponent, we add as well a unique SkeletalMeshComponent
+				// (so we can remove them later on in a consistent fashion)
+				check(CharacterMovementComponents.Num() == SkeletalMeshComponents.Num());
+
+				PreTickFunction.AddPrerequisite(Actor, CharacterMovementComponent->PrimaryComponentTick);
+				SkeletalMeshComponent->PrimaryComponentTick.AddPrerequisite(Actor, PreTickFunction);
+
+				if (bIsFirstInjectedActor)
+				{
+					PostTickFunction.AddPrerequisite(Actor, SkeletalMeshComponent->PrimaryComponentTick);
+				}
+				else 
+				{
+					SkeletalMeshComponent->PrimaryComponentTick.AddPrerequisite(Actor, PostTickFunction);
+				}
+			}
+			else
+			{
+				UE_LOG(LogChimera, Error, TEXT("FIsland::InjectToActor requires AActor %s to have a USkeletalMeshComponent to work!"), *Actor->GetName());
+			}
+		}
+		else
+		{
+			UE_LOG(LogChimera, Error, TEXT("FIsland::InjectToActor requires AActor %s to have a UCharacterMovementComponent to work!"), *Actor->GetName());
+		}
+	}
+}
+
+void FIsland::AddSearchContext(const FSearchContext& SearchContext)
+{
+#if DO_CHECK
+	check(SearchContext.IsValid());
+	check(IsInGameThread());
+
+	for (const FSearchContext& ContainedSearchContext : SearchContexts)
+	{
+		check(!ContainedSearchContext.IsEquivalent(SearchContext));
+	}
+#endif // DO_CHECK
+	SearchContexts.Add(SearchContext);
+}
+
+void FIsland::Uninject()
+{
+	// Called by UChimeraSubsystem::Tick when there aren't animation jobs flying. No need to FScopeLock Lock(&Mutex);
+	check(IsInGameThread());
+
+	check(CharacterMovementComponents.Num() == SkeletalMeshComponents.Num());
+
+	for (int32 ActorIndex = 0; ActorIndex < CharacterMovementComponents.Num(); ++ActorIndex)
+	{
+		UCharacterMovementComponent* CharacterMovementComponent = CharacterMovementComponents[ActorIndex].Get();
+		USkeletalMeshComponent* SkeletalMeshComponent = SkeletalMeshComponents[ActorIndex].Get();
+
+		if (CharacterMovementComponent)
+		{
+			check(SkeletalMeshComponent);
+
+			AActor* Actor = CharacterMovementComponent->GetOwner();
+			check(Actor);
+
+			PreTickFunction.RemovePrerequisite(Actor, CharacterMovementComponent->PrimaryComponentTick);
+			SkeletalMeshComponent->PrimaryComponentTick.RemovePrerequisite(Actor, PreTickFunction);
+
+			const bool bIsFirstInjectedActor = ActorIndex == 0;
+			if (bIsFirstInjectedActor)
+			{
+				PostTickFunction.RemovePrerequisite(Actor, SkeletalMeshComponent->PrimaryComponentTick);
+			}
+			else
+			{
+				SkeletalMeshComponent->PrimaryComponentTick.RemovePrerequisite(Actor, PostTickFunction);
+			}
+		}
+		else
+		{
+			check(!SkeletalMeshComponent);
+		}
+	}
+
+	CharacterMovementComponents.Reset();
+	SkeletalMeshComponents.Reset();
+
+	SearchContexts.Reset();
+	SearchResults.Reset();
+	bSearchPerfomed = false;
+}
+
+bool FIsland::IsUninjected() const
+{
+	return SkeletalMeshComponents.IsEmpty();
+}
+
+bool FIsland::DoSearch_AnyThread(UObject* AnimInstance, const FPoseSearchContinuingProperties& ContinuingProperties, FChimeraBlueprintResult& Result)
+{
+	bool bDoPerfomSearch = false;
+	{
+		// thread safety note!
+		// goal:	avoiding deadlock between SearchResultsMutex lock and waiting for UAnimInstance::HandleExistingParallelEvaluationTask.
+		// why:		UPoseSearchLibrary::MotionMatch could call via AnimInstance GetProxyOnAnyThread<FAnimInstanceProxy>() that, if on GameThread, 
+		//			could call UAnimInstance::HandleExistingParallelEvaluationTask 
+		// fix:		avoid UPoseSearchLibrary::MotionMatch calls wrapped by any lock, at the cost of eventually (by design should be NEVER) performing the searches twice
+		//			By design we should inject ticks dependencies (added by FIsland::InjectToActor via AddTickPrerequisiteComponent), so concurrently 
+		//			fly of FIsland within the same island that requires searches is forbidden
+		FScopeLock Lock(&SearchResultsMutex);
+		bDoPerfomSearch = !bSearchPerfomed;
+	}
+
+	if (bDoPerfomSearch)
+	{
+		FMemMark Mark(FMemStack::Get());
+
+		TArray<UAnimInstance*, TInlineAllocator<UE::PoseSearch::PreallocatedRolesNum, TMemStackAllocator<>>> AnimInstances;
+		TArray<const UE::PoseSearch::IPoseHistory*, TInlineAllocator<UE::PoseSearch::PreallocatedRolesNum, TMemStackAllocator<>>> PoseHistories;
+		TArray<UE::PoseSearch::FSearchResult, TMemStackAllocator<>> PoseSearchResults;
+
+		// SearchContexts are modified only by UChimeraSubsystem::Tick and constant otherwise, so it's safe to access them in a threaded enviroment without locks
+		PoseSearchResults.SetNum(SearchContexts.Num());
+
+		for (int32 SearchIndex = 0; SearchIndex < SearchContexts.Num(); ++SearchIndex)
+		{
+			const FSearchContext& SearchContext = SearchContexts[SearchIndex];
+			const UPoseSearchDatabase* Database = SearchContext.Database.Get();
+			if (!Database)
+			{
+				UE_LOG(LogChimera, Error, TEXT("FIsland::DoSearch_AnyThread invalid context database"));
+				return false;
+			}
+
+			if (!Database->Schema)
+			{
+				UE_LOG(LogChimera, Error, TEXT("FIsland::DoSearch_AnyThread invalid schema for context database %s"), *Database->GetName());
+				return false;
+			}
+
+			AnimInstances.Reset();
+			for (const TWeakObjectPtr<UAnimInstance>& AnimInstancePtr : SearchContext.AnimInstances)
+			{
+				UAnimInstance* SearchContextAnimInstance = AnimInstancePtr.Get();
+				if (!SearchContextAnimInstance)
+				{
+					UE_LOG(LogChimera, Error, TEXT("FIsland::DoSearch_AnyThread null anim instance"));
+					return false;
+				}
+
+				AnimInstances.Add(SearchContextAnimInstance);
+			}
+
+			PoseHistories.Reset();
+			for (const FAnimNode_PoseSearchHistoryCollector_Base* HistoryCollector : SearchContext.HistoryCollectors)
+			{
+				PoseHistories.Add(&HistoryCollector->GetPoseHistory());
+			}
+
+			const UObject* AssetsToSearch[] = { Database };
+			FPoseSearchFutureProperties PoseSearchFutureProperties;
+
+			FPoseSearchContinuingProperties ContinuingPropertiesToUse = SearchContext.ContinuingProperties;
+
+			// @todo: implemennt this carefully, since it'll cause thread safety concerns depending on which actors is calling the query first 
+			// and if their animations get integrated by different play rates
+			
+			//if (ContinuingProperties.PlayingAsset)
+			//{
+			//	if (const UMultiAnimAsset* MultiAnimAsset = Cast<UMultiAnimAsset>(ContinuingProperties.PlayingAsset))
+			//	{
+			//		// @todo: if they are compatible... etcetc MultiAnimAsset->GetNumRoles
+			//		ContinuingPropertiesToUse = ContinuingProperties;
+			//	}
+			//	else if (const UMultiAnimAsset* SearchContextMultiAnimAsset = Cast<UMultiAnimAsset>(SearchContext.ContinuingProperties.PlayingAsset))
+			//	{
+			//		for (int32 RoleIndex = 0; RoleIndex < SearchContextMultiAnimAsset->GetNumRoles(); ++RoleIndex)
+			//		{
+			//			if (SearchContextMultiAnimAsset->GetAnimationAsset(SearchContextMultiAnimAsset->GetRole(RoleIndex)) == ContinuingProperties.PlayingAsset)
+			//			{
+			//				ContinuingPropertiesToUse.PlayingAssetAccumulatedTime = ContinuingProperties.PlayingAssetAccumulatedTime;
+			//				break;
+			//			}
+			//		}
+			//	}
+			//	else if (!SearchContext.ContinuingProperties.PlayingAsset)
+			//	{
+			//		// @todo: should we scan all the database iterating over the UMultiAnimAsset to look for ContinuingProperties.PlayingAsset?
+			//	}
+			//	else
+			//	{
+			//	}
+			//}
+
+			// @todo: we could perform multiple UPoseSearchLibrary::MotionMatch in parallel!
+			const UE::PoseSearch::FSearchResult PoseSearchResult = UPoseSearchLibrary::MotionMatch(AnimInstances, SearchContext.Roles,
+				PoseHistories, AssetsToSearch, ContinuingPropertiesToUse, PoseSearchFutureProperties);
+
+			if (PoseSearchResult.PoseCost.GetTotalCost() < SearchContext.MaxCost)
+			{
+				PoseSearchResults[SearchIndex] = PoseSearchResult;
+			}
+		}
+
+		if (!PoseSearchResults.IsEmpty())
+		{
+			// locking to update SearchResults and bSearchPerfomed
+			FScopeLock Lock(&SearchResultsMutex);
+
+			if (!bSearchPerfomed)
+			{
+				InitSearchResults(SearchResults, PoseSearchResults, SearchContexts);
+
+				bSearchPerfomed = true;
+			}
+			else
+			{
+				UE_LOG(LogChimera, Warning, TEXT("FIsland::DoSearch_AnyThread performance warning: performed duplicated search :("));
+				
+#if DO_CHECK
+				TArray<FSearchResult> CompareSearchResults;
+				InitSearchResults(CompareSearchResults, PoseSearchResults, SearchContexts);
+
+				if (CompareSearchResults != SearchResults)
+				{
+					UE_LOG(LogChimera, Error, TEXT("FIsland::DoSearch_AnyThread duplicated search differs from original one. Searches are NOT deterministic!"));
+				}
+#endif // DO_CHECK
+			}
+	
+			// calling this funtion within this scope since we already locked SearchResultsMutex
+			return GetResult_AnyThread(AnimInstance, Result);
+		}
+	}
+
+	return GetResult_AnyThread(AnimInstance, Result);
+}
+
+bool FIsland::GetResult_AnyThread(UObject* AnimInstance, FChimeraBlueprintResult& Result)
+{
+	// locking to read SearchResults
+	FScopeLock Lock(&SearchResultsMutex);
+
+	// looking for AnimInstance in SearchResults to fill up Result
+	for (const FSearchResult& SearchResult : SearchResults)
+	{
+		const FSearchContext& SearchContext = SearchContexts[SearchResult.SearchIndex];
+		for (int32 AnimInstanceIndex = 0; AnimInstanceIndex < SearchContext.AnimInstances.Num(); ++AnimInstanceIndex)
+		{
+			if (SearchContext.AnimInstances[AnimInstanceIndex].Get() == AnimInstance)
+			{
+				const UPoseSearchDatabase* Database = SearchResult.Database.Get();
+				check(Database);
+
+				const UE::PoseSearch::FSearchIndexAsset* SearchIndexAsset = SearchResult.GetSearchIndexAsset();
+				check(SearchIndexAsset);
+
+				const FPoseSearchDatabaseAnimationAssetBase* DatabaseAnimationAssetBase = Database->GetDatabaseAnimationAsset<FPoseSearchDatabaseAnimationAssetBase>(*SearchIndexAsset);
+				check(DatabaseAnimationAssetBase);
+
+				const UE::PoseSearch::FRole Role = SearchContext.Roles[AnimInstanceIndex];
+
+				Result.SelectedAnimation = DatabaseAnimationAssetBase->GetAnimationAsset();
+				Result.SelectedTime = SearchResult.AssetTime;
+				Result.bIsContinuingPoseSearch = SearchResult.bIsContinuingPoseSearch;
+				Result.bLoop = SearchIndexAsset->IsLooping();
+				Result.bIsMirrored = SearchIndexAsset->IsMirrored();
+				Result.BlendParameters = SearchIndexAsset->GetBlendParameters();
+				Result.SelectedDatabase = Database;
+				Result.SearchCost = SearchResult.PoseCost.GetTotalCost();
+				Result.Role = Role;
+
+				if (SearchResult.FullAlignedActorRootBoneTransforms.IsValidIndex(AnimInstanceIndex))
+				{
+					Result.FullAlignedActorRootBoneTransform = SearchResult.FullAlignedActorRootBoneTransforms[AnimInstanceIndex];
+				}
+
+				//Result.bIsFromContinuingPlaying = SearchResult.bIsContinuingPoseSearch;
+
+				// figuring out the WantedPlayRate
+				Result.WantedPlayRate = 1.f;
+				//if (Future.Animation && Future.IntervalTime > 0.f)
+				//{
+				//	if (const UPoseSearchFeatureChannel_PermutationTime* PermutationTimeChannel = Database->Schema->FindFirstChannelOfType<UPoseSearchFeatureChannel_PermutationTime>())
+				//	{
+				//		const FSearchIndex& SearchIndex = Database->GetSearchIndex();
+				//		if (!SearchIndex.IsValuesEmpty())
+				//		{
+				//			TConstArrayView<float> ResultData = Database->GetSearchIndex().GetPoseValues(SearchResult.PoseIdx);
+				//			const float ActualIntervalTime = PermutationTimeChannel->GetPermutationTime(ResultData);
+				//			ProviderResult.WantedPlayRate = ActualIntervalTime / Future.IntervalTime;
+				//		}
+				//	}
+				//}
+
+				// we found our AnimInstance: we can stop searching
+				return true;
+			}
+		}
+	}
+
+	return false;
+}
+
+const FSearchResult* FIsland::FindSearchResult(const FSearchContext& SearchContext) const
+{
+	// called only by UChimeraSubsystem::Tick via UChimeraSubsystem::PopulateContinuingProperties so no need to lock SearchResultsMutex to protect the read of SearchResults
+	check(IsInGameThread());
+
+	// searching for InSearchContext in all the SearchContexts referenced by valid active SearchResults
+	for (const FSearchResult& SearchResult : SearchResults)
+	{
+		const FSearchContext& LocalSearchContext = SearchContexts[SearchResult.SearchIndex];
+		if (LocalSearchContext.Database == SearchContext.Database &&
+			LocalSearchContext.AnimInstances == SearchContext.AnimInstances &&
+			LocalSearchContext.Roles == SearchContext.Roles)
+		{
+			return &SearchResult;
+		}
+	}
+	return nullptr;
+}
+
+} // namespace UE::Chimera
diff --git a/Engine/Plugins/Experimental/Animation/Chimera/Source/Runtime/Private/ChimeraIslandComponent.cpp b/Engine/Plugins/Experimental/Animation/Chimera/Source/Runtime/Private/ChimeraIslandComponent.cpp
deleted file mode 100644
diff --git a/Engine/Plugins/Experimental/Animation/Chimera/Source/Runtime/Private/ChimeraSubsystem.cpp b/Engine/Plugins/Experimental/Animation/Chimera/Source/Runtime/Private/ChimeraSubsystem.cpp
--- a/Engine/Plugins/Experimental/Animation/Chimera/Source/Runtime/Private/ChimeraSubsystem.cpp
+++ b/Engine/Plugins/Experimental/Animation/Chimera/Source/Runtime/Private/ChimeraSubsystem.cpp
@@ -2,7 +2,7 @@
 
 #include "Chimera/ChimeraSubsystem.h"
 #include "Chimera/ChimeraDefines.h"
-#include "Chimera/ChimeraIslandComponent.h"
+#include "Chimera/ChimeraIsland.h"
 #include "PoseSearch/PoseSearchDatabase.h"
 #include "PoseSearch/PoseSearchRole.h"
 #include "PoseSearch/PoseSearchSchema.h"
diff --git a/Engine/Plugins/Experimental/Animation/Chimera/Source/Runtime/Public/Chimera/ChimeraIsland.h b/Engine/Plugins/Experimental/Animation/Chimera/Source/Runtime/Public/Chimera/ChimeraIsland.h
new file mode 100644
--- /dev/null
+++ b/Engine/Plugins/Experimental/Animation/Chimera/Source/Runtime/Public/Chimera/ChimeraIsland.h
@@ -0,0 +1,102 @@
+// Copyright Epic Games, Inc. All Rights Reserved.
+
+#pragma once
+
+#include "PoseSearch/PoseSearchLibrary.h"
+
+struct FChimeraBlueprintResult;
+struct FPoseSearchContinuingProperties;
+class UCharacterMovementComponent;
+
+namespace UE::Chimera
+{
+struct FIsland;
+
+struct FSearchContext
+{
+	// @todo: since AnimInstances, HistoryCollectors, Roles, and BroadPhaseRadiuses have the same cardinality,
+	// add an array of structs instead of many arrays
+	TArray<TWeakObjectPtr<UAnimInstance>> AnimInstances;
+	// @todo: perhaps use the animnode of the pose history collector instead
+	TArray<const FAnimNode_PoseSearchHistoryCollector_Base*> HistoryCollectors;
+	TArray<UE::PoseSearch::FRole> Roles;
+
+#if ENABLE_DRAW_DEBUG
+	TArray<float> BroadPhaseRadiuses;
+#endif // ENABLE_DRAW_DEBUG
+
+	float MaxCost = MAX_flt;
+
+	TWeakObjectPtr<const UPoseSearchDatabase> Database;
+	FPoseSearchContinuingProperties ContinuingProperties;
+
+	bool IsValid() const;
+	bool IsEquivalent(const FSearchContext& Other) const;
+};
+
+struct FSearchResult : public UE::PoseSearch::FSearchResult
+{
+	int32 SearchIndex = INDEX_NONE;
+	TArray<FTransform, TInlineAllocator<UE::PoseSearch::PreallocatedRolesNum>> FullAlignedActorRootBoneTransforms;
+
+	bool operator==(const FSearchResult& Other) const;
+};
+
+// FIsland contains ticks functions injected between the interacting actors UCharacterMovementComponent and USkeletalMeshComponent
+// to create a execution threading fence to be able to perform motion matching searches between the involved characters in a thread safe manner.
+// Look at UChimeraSubsystem "Execution model and threading details" for additional information
+struct FIsland
+{
+	UE_NONCOPYABLE(FIsland);
+	
+	FIsland(ULevel* Level);
+	~FIsland();
+
+	bool DoSearch_AnyThread(UObject* AnimInstance, const FPoseSearchContinuingProperties& ContinuingProperties, FChimeraBlueprintResult& Result);
+
+	const TArray<TWeakObjectPtr<UCharacterMovementComponent>>& GetCharacterMovementComponents() const { return CharacterMovementComponents; }
+	const TArray<TWeakObjectPtr<USkeletalMeshComponent>>& GetSkeletalMeshComponents() const { return SkeletalMeshComponents; }
+	const TArray<UE::Chimera::FSearchContext>& GetSearchContexts() const { return SearchContexts; }
+	const TArray<UE::Chimera::FSearchResult>& GetSearchResults() const { return SearchResults; }
+
+	bool IsUninjected() const;
+	void InjectToActor(AActor* Actor);
+	void Uninject();
+
+	const UE::Chimera::FSearchResult* FindSearchResult(const UE::Chimera::FSearchContext& SearchContext) const;
+	void AddSearchContext(const UE::Chimera::FSearchContext& SearchContext);
+
+	void DebugDraw(const FColor& Color = FColor::Red) const;
+
+private:
+	bool GetResult_AnyThread(UObject* AnimInstance, FChimeraBlueprintResult& Result);
+
+	struct FPreTickFunction : public FTickFunction
+	{
+		virtual void ExecuteTick(float DeltaTime, enum ELevelTick TickType, ENamedThreads::Type CurrentThread, const FGraphEventRef& MyCompletionGraphEvent) override;
+		virtual FString DiagnosticMessage() override { return TEXT("FPreTickFunction"); }
+		FIsland* Island = nullptr;
+	};
+	FPreTickFunction PreTickFunction;
+
+	struct FPostTickFunction : public FTickFunction
+	{
+		virtual void ExecuteTick(float DeltaTime, enum ELevelTick TickType, ENamedThreads::Type CurrentThread, const FGraphEventRef& MyCompletionGraphEvent) override;
+		virtual FString DiagnosticMessage() override { return TEXT("FPostTickFunction"); }
+	};
+	FPostTickFunction PostTickFunction;
+
+	TArray<TWeakObjectPtr<UCharacterMovementComponent>> CharacterMovementComponents;
+	TArray<TWeakObjectPtr<USkeletalMeshComponent>> SkeletalMeshComponents;
+	
+	// there's one FSearchContext for each search we need to perform (including all the possible roles permutations). Added by UChimeraSubsystem::Tick
+	TArray<UE::Chimera::FSearchContext> SearchContexts;
+
+	// SearchResults contains only the best results, and it has not necessarly the same cardinality as SearchContexts. usually SearchResults.Num() < SearchContexts.Num()
+	TArray<UE::Chimera::FSearchResult> SearchResults;
+	bool bSearchPerfomed = false;
+
+	FCriticalSection SearchResultsMutex;
+};
+
+} // namespace UE::Chimera
diff --git a/Engine/Plugins/Experimental/Animation/Chimera/Source/Runtime/Public/Chimera/ChimeraIslandComponent.h b/Engine/Plugins/Experimental/Animation/Chimera/Source/Runtime/Public/Chimera/ChimeraIslandComponent.h
deleted file mode 100644
